OnlineServicesOSSAdapter

diff --git a/Engine/Plugins/Online/Experimental/OnlineServicesOSSAdapter/OnlineServicesOSSAdapter.uplugin b/Engine/Plugins/Online/Experimental/OnlineServicesOSSAdapter/OnlineServicesOSSAdapter.uplugin
new file mode 100644
--- /dev/null
+++ b/Engine/Plugins/Online/Experimental/OnlineServicesOSSAdapter/OnlineServicesOSSAdapter.uplugin
@@ -0,0 +1,30 @@
+{
+	"FileVersion": 3,
+
+	"FriendlyName": "Online Services OSS Adapter",
+	"Version": 1,
+	"VersionName": "1.0",
+	"Description": "Online Services adapter for Online Subsystem implementations.",
+	"Category": "Online Platform",
+	"CreatedBy": "Epic Games, Inc.",
+	"CreatedByURL": "http://epicgames.com",
+	"EnabledByDefault": false,
+
+	"Modules": [
+		{
+			"Name": "OnlineServicesOSSAdapter",
+			"Type": "Runtime",
+			"LoadingPhase": "PostConfigInit"
+		}
+	],
+	"Plugins": [
+		{
+			"Name": "OnlineServices",
+			"Enabled": true
+		},
+		{
+			"Name": "OnlineSubsystem",
+			"Enabled": true
+		}
+	]
+}
\ No newline at end of file
diff --git a/Engine/Plugins/Online/Experimental/OnlineServicesOSSAdapter/Source/OnlineServicesOSSAdapter.Build.cs b/Engine/Plugins/Online/Experimental/OnlineServicesOSSAdapter/Source/OnlineServicesOSSAdapter.Build.cs
new file mode 100644
--- /dev/null
+++ b/Engine/Plugins/Online/Experimental/OnlineServicesOSSAdapter/Source/OnlineServicesOSSAdapter.Build.cs
@@ -0,0 +1,25 @@
+// Copyright Epic Games, Inc. All Rights Reserved.
+
+using UnrealBuildTool;
+using System.IO;
+
+public class OnlineServicesOSSAdapter : ModuleRules
+{
+	public OnlineServicesOSSAdapter(ReadOnlyTargetRules Target) : base(Target)
+    {
+		PublicDependencyModuleNames.AddRange(
+			new string[] {
+				"Core",
+				"OnlineServicesInterface",
+				"OnlineSubsystem"
+			}
+		);
+
+		PrivateDependencyModuleNames.AddRange(
+			new string[] {
+				"CoreOnline",
+				"OnlineServicesCommon"
+			}
+		);
+	}
+}
diff --git a/Engine/Plugins/Online/Experimental/OnlineServicesOSSAdapter/Source/Private/Online/AuthOSSAdapter.cpp b/Engine/Plugins/Online/Experimental/OnlineServicesOSSAdapter/Source/Private/Online/AuthOSSAdapter.cpp
new file mode 100644
--- /dev/null
+++ b/Engine/Plugins/Online/Experimental/OnlineServicesOSSAdapter/Source/Private/Online/AuthOSSAdapter.cpp
@@ -0,0 +1,251 @@
+// Copyright Epic Games, Inc. All Rights Reserved.
+
+#pragma once
+
+#include "Online/AuthOSSAdapter.h"
+
+#include "Online/OnlineServicesOSSAdapter.h"
+#include "Online/OnlineIdOSSAdapter.h"
+#include "Online/OnlineErrorDefinitions.h"
+
+#include "OnlineSubsystem.h"
+#include "Interfaces/OnlineIdentityInterface.h"
+
+
+namespace UE::Online {
+
+void FAuthOSSAdapter::PostInitialize()
+{
+	Super::PostInitialize();
+
+	if (IOnlineIdentityPtr Identity = GetIdentityInterface())
+	{
+		for (int LocalPlayerNum = 0; LocalPlayerNum < MAX_LOCAL_PLAYERS; ++LocalPlayerNum)
+		{
+			OnLoginStatusChangedHandle[LocalPlayerNum] = Identity->OnLoginStatusChangedDelegates[LocalPlayerNum].AddLambda(
+				[WeakThis = TWeakPtr<IAuth>(AsShared()), this](int32 LocalUserNum, ::ELoginStatus::Type OldStatus, ::ELoginStatus::Type NewStatus, const FUniqueNetId& NewId)
+				{
+					TSharedPtr<IAuth> PinnedThis = WeakThis.Pin();
+					if (PinnedThis.IsValid())
+					{
+						FLoginStatusChanged LoginStatusChanged;
+						if (NewId.IsValid())
+						{
+							LoginStatusChanged.LocalUserId = static_cast<FOnlineServicesOSSAdapter&>(Services).GetAccountIdRegistry().FindOrAddHandle(NewId.AsShared());
+						}
+						switch (OldStatus)
+						{
+							case ::ELoginStatus::NotLoggedIn:
+								LoginStatusChanged.PreviousStatus = ELoginStatus::NotLoggedIn;
+								break;
+							case ::ELoginStatus::UsingLocalProfile:
+								LoginStatusChanged.PreviousStatus = ELoginStatus::UsingLocalProfile;
+								break;
+							case ::ELoginStatus::LoggedIn:
+								LoginStatusChanged.PreviousStatus = ELoginStatus::LoggedIn;
+								break;
+						}
+						switch (NewStatus)
+						{
+							case ::ELoginStatus::NotLoggedIn:
+								LoginStatusChanged.CurrentStatus = ELoginStatus::NotLoggedIn;
+								break;
+							case ::ELoginStatus::UsingLocalProfile:
+								LoginStatusChanged.CurrentStatus = ELoginStatus::UsingLocalProfile;
+								break;
+							case ::ELoginStatus::LoggedIn:
+								LoginStatusChanged.CurrentStatus = ELoginStatus::LoggedIn;
+								break;
+						}
+						OnLoginStatusChangedEvent.Broadcast(LoginStatusChanged);
+					}
+				});
+		}
+	}
+}
+
+void FAuthOSSAdapter::PreShutdown()
+{
+	if (IOnlineIdentityPtr Identity = GetIdentityInterface())
+	{
+		for (int LocalPlayerNum = 0; LocalPlayerNum < MAX_LOCAL_PLAYERS; ++LocalPlayerNum)
+		{
+			Identity->ClearOnLoginStatusChangedDelegate_Handle(LocalPlayerNum, OnLoginStatusChangedHandle[LocalPlayerNum]);
+		}
+	}
+
+	Super::PreShutdown();
+}
+
+TOnlineAsyncOpHandle<FAuthLogin> FAuthOSSAdapter::Login(FAuthLogin::Params&& Params)
+{
+	TSharedRef<TOnlineAsyncOp<FAuthLogin>> Op = GetOp<FAuthLogin>(MoveTemp(Params));
+
+	if (!Op->IsReady())
+	{
+		int32 LocalUserIndex = FPlatformMisc::GetUserIndexForPlatformUser(Op->GetParams().PlatformUserId);
+
+		if (LocalUserIndex < 0 || LocalUserIndex >= MAX_LOCAL_PLAYERS)
+		{
+			Op->SetError(Errors::InvalidParams());
+			return Op->GetHandle();
+		}
+
+		Op->Then([this, LocalUserIndex](TOnlineAsyncOp<FAuthLogin>& Op, TPromise<TOnlineResult<FAuthLogin>>&& Result)
+			{
+				FOnlineAccountCredentials Credentials;
+				Credentials.Type = Op.GetParams().CredentialsType;
+				Credentials.Id = Op.GetParams().CredentialsId;
+				Credentials.Token = Op.GetParams().CredentialsToken.Get<FString>(); // TODO: handle binary token
+
+				TSharedPtr<TPair<FDelegateHandle, TPromise<TOnlineResult<FAuthLogin>>>> ResultPtr = MakeShared<TPair<FDelegateHandle, TPromise<TOnlineResult<FAuthLogin>>>>(FDelegateHandle(), MoveTemp(Result));
+				ResultPtr->Get<FDelegateHandle>() = GetIdentityInterface()->OnLoginCompleteDelegates[LocalUserIndex].AddLambda(
+					[WeakThis = TWeakPtr<IAuth>(AsShared()), this, ResultPtr, PlatformUserId = Op.GetParams().PlatformUserId, LocalUserIndex](int32 LocalUserNum, bool bWasSuccessful, const FUniqueNetId& UserId, const FString& Error) mutable
+					{
+						if (LocalUserIndex != LocalUserNum)
+						{
+							return;
+						}
+
+						TSharedPtr<IAuth> PinnedThis = WeakThis.Pin();
+						if (PinnedThis.IsValid())
+						{
+							TPromise<TOnlineResult<FAuthLogin>>& OnlineResult = ResultPtr->Get<TPromise<TOnlineResult<FAuthLogin>>>();
+							if (bWasSuccessful)
+							{
+								FOnlineAccountIdHandle Handle = static_cast<FOnlineServicesOSSAdapter&>(Services).GetAccountIdRegistry().FindOrAddHandle(UserId.AsShared());
+								FAuthLogin::Result Result = { MakeShared<FAccountInfo>() };
+								Result.AccountInfo->PlatformUserId = PlatformUserId;
+								Result.AccountInfo->UserId = Handle;
+								Result.AccountInfo->LoginStatus = ELoginStatus::LoggedIn;
+								OnlineResult.EmplaceValue(MoveTemp(Result));
+							}
+							else
+							{
+								FOnlineError V2Error = Errors::Unknown(); // TODO: V1 to V2 error conversion/error from string conversion
+								OnlineResult.EmplaceValue(MoveTemp(V2Error));
+							}
+						}
+
+						GetIdentityInterface()->OnLoginCompleteDelegates[LocalUserIndex].Remove(ResultPtr->Get<FDelegateHandle>());
+					});
+
+				GetIdentityInterface()->Login(LocalUserIndex, Credentials);
+			})
+			.Then([this](TOnlineAsyncOp<FAuthLogin>& Op, const TOnlineResult<FAuthLogin>& Result)
+			{
+				if (Result.IsError())
+				{
+					Op.SetError(CopyTemp(Result.GetErrorValue()));
+				}
+				else
+				{
+					Op.SetResult(CopyTemp(Result.GetOkValue()));
+				}
+			})
+			.Enqueue(GetSerialQueue());
+	}
+	return Op->GetHandle();
+}
+
+TOnlineAsyncOpHandle<FAuthLogout> FAuthOSSAdapter::Logout(FAuthLogout::Params&& Params)
+{
+	TSharedRef<TOnlineAsyncOp<FAuthLogout>> Op = GetOp<FAuthLogout>(MoveTemp(Params));
+
+	if (!Op->IsReady())
+	{
+		int32 LocalUserIndex = GetLocalUserNum(Op->GetParams().LocalUserId);
+
+		if (LocalUserIndex < 0 || LocalUserIndex >= MAX_LOCAL_PLAYERS)
+		{
+			Op->SetError(Errors::InvalidParams());
+			return Op->GetHandle();
+		}
+
+		Op->Then([this, LocalUserIndex](TOnlineAsyncOp<FAuthLogout>& Op, TPromise<TOnlineResult<FAuthLogout>>&& Result)
+			{
+				TSharedPtr<TPair<FDelegateHandle, TPromise<TOnlineResult<FAuthLogout>>>> ResultPtr = MakeShared<TPair<FDelegateHandle, TPromise<TOnlineResult<FAuthLogout>>>>(FDelegateHandle(), MoveTemp(Result));
+				ResultPtr->Get<FDelegateHandle>() = GetIdentityInterface()->OnLogoutCompleteDelegates[LocalUserIndex].AddLambda(
+					[WeakThis = TWeakPtr<IAuth>(AsShared()), this, ResultPtr, LocalUserIndex](int32 LocalUserNum, bool bWasSuccessful) mutable
+				{
+					if (LocalUserIndex != LocalUserNum)
+					{
+						return;
+					}
+
+					TPromise<TOnlineResult<FAuthLogout>>& OnlineResult = ResultPtr->Get<TPromise<TOnlineResult<FAuthLogout>>>();
+
+					TSharedPtr<IAuth> PinnedThis = WeakThis.Pin();
+					if (PinnedThis.IsValid())
+					{
+						if (bWasSuccessful)
+						{
+							FAuthLogout::Result Result;
+							OnlineResult.EmplaceValue(MoveTemp(Result));
+						}
+						else
+						{
+							FOnlineError V2Error = Errors::Unknown(); // TODO: V1 to V2 error conversion/error from string conversion
+							OnlineResult.EmplaceValue(MoveTemp(V2Error));
+						}
+					}
+
+					GetIdentityInterface()->OnLogoutCompleteDelegates[LocalUserIndex].Remove(ResultPtr->Get<FDelegateHandle>());
+				});
+
+				GetIdentityInterface()->Logout(LocalUserIndex);
+			})
+			.Then([this](TOnlineAsyncOp<FAuthLogout>& Op, const TOnlineResult<FAuthLogout>& Result)
+				{
+					if (Result.IsError())
+					{
+						Op.SetError(CopyTemp(Result.GetErrorValue()));
+					}
+					else
+					{
+						Op.SetResult(CopyTemp(Result.GetOkValue()));
+					}
+				})
+				.Enqueue(GetSerialQueue());
+	}
+	return Op->GetHandle();
+}
+
+TOnlineAsyncOpHandle<FAuthGenerateAuth> FAuthOSSAdapter::GenerateAuth(FAuthGenerateAuth::Params&& Params)
+{
+	// TODO: If no cached auth token available, call Login with appropriate parameters to generate an auth token
+
+	return Super::GenerateAuth(MoveTemp(Params));
+}
+
+TOnlineResult<FAuthGetAccountByPlatformUserId> FAuthOSSAdapter::GetAccountByPlatformUserId(FAuthGetAccountByPlatformUserId::Params&& Params)
+{
+	return Super::GetAccountByPlatformUserId(MoveTemp(Params));
+}
+
+TOnlineResult<FAuthGetAccountByAccountId> FAuthOSSAdapter::GetAccountByAccountId(FAuthGetAccountByAccountId::Params&& Params)
+{
+	return Super::GetAccountByAccountId(MoveTemp(Params));
+}
+
+FUniqueNetIdRef FAuthOSSAdapter::GetUniqueNetId(FOnlineAccountIdHandle AccountIdHandle) const
+{
+	return static_cast<FOnlineServicesOSSAdapter&>(Services).GetAccountIdRegistry().GetIdValue(AccountIdHandle);
+}
+
+int32 FAuthOSSAdapter::GetLocalUserNum(FOnlineAccountIdHandle AccountIdHandle) const
+{
+	FUniqueNetIdRef UniqueNetId = GetUniqueNetId(AccountIdHandle);
+	if (UniqueNetId->IsValid())
+	{
+		return GetIdentityInterface()->GetLocalUserNumFromPlatformUserId(GetIdentityInterface()->GetPlatformUserIdFromUniqueNetId(*UniqueNetId));
+	}
+	return INDEX_NONE;
+}
+
+IOnlineIdentityPtr FAuthOSSAdapter::GetIdentityInterface() const
+{
+	return static_cast<FOnlineServicesOSSAdapter&>(Services).GetSubsystem().GetIdentityInterface();
+}
+
+/* UE::Online */ }
diff --git a/Engine/Plugins/Online/Experimental/OnlineServicesOSSAdapter/Source/Private/Online/OnlineServicesOSSAdapter.cpp b/Engine/Plugins/Online/Experimental/OnlineServicesOSSAdapter/Source/Private/Online/OnlineServicesOSSAdapter.cpp
new file mode 100644
--- /dev/null
+++ b/Engine/Plugins/Online/Experimental/OnlineServicesOSSAdapter/Source/Private/Online/OnlineServicesOSSAdapter.cpp
@@ -0,0 +1,32 @@
+// Copyright Epic Games, Inc. All Rights Reserved.
+
+#include "Online/OnlineServicesOSSAdapter.h"
+
+#include "Online/AuthOSSAdapter.h"
+
+#include "OnlineSubsystem.h"
+
+namespace UE::Online {
+
+FOnlineServicesOSSAdapter::FOnlineServicesOSSAdapter(EOnlineServices InServicesType, const FString& InConfigName, IOnlineSubsystem* InSubsystem)
+	: FOnlineServicesCommon(InConfigName)
+	, ServicesType(InServicesType)
+	, Subsystem(InSubsystem)
+{
+}
+
+void FOnlineServicesOSSAdapter::RegisterComponents()
+{
+	Components.Register<FAuthOSSAdapter>(*this);
+
+	FOnlineServicesCommon::RegisterComponents();
+}
+
+void FOnlineServicesOSSAdapter::Initialize()
+{
+	AccountIdRegistry = static_cast<FOnlineUniqueNetIdRegistry*>(FOnlineIdRegistryRegistry::Get().GetAccountIdRegistry(ServicesType));
+
+	FOnlineServicesCommon::Initialize();
+}
+
+/* UE::Online*/ }
diff --git a/Engine/Plugins/Online/Experimental/OnlineServicesOSSAdapter/Source/Private/Online/OnlineServicesOSSAdapterModule.cpp b/Engine/Plugins/Online/Experimental/OnlineServicesOSSAdapter/Source/Private/Online/OnlineServicesOSSAdapterModule.cpp
new file mode 100644
--- /dev/null
+++ b/Engine/Plugins/Online/Experimental/OnlineServicesOSSAdapter/Source/Private/Online/OnlineServicesOSSAdapterModule.cpp
@@ -0,0 +1,25 @@
+// Copyright Epic Games, Inc. All Rights Reserved.
+
+#include "CoreMinimal.h"
+#include "Modules/ModuleInterface.h"
+#include "Modules/ModuleManager.h"
+
+#include "Online/OnlineServicesRegistry.h"
+
+namespace UE::Online
+{
+
+class FOnlineServicesOSSAdapterModule : public IModuleInterface
+{
+public:
+	virtual void StartupModule() override;
+protected:
+};
+
+void FOnlineServicesOSSAdapterModule::StartupModule()
+{
+}
+
+/* UE::Online */ }
+
+IMPLEMENT_MODULE(UE::Online::FOnlineServicesOSSAdapterModule, OnlineServicesOSSAdapter);
diff --git a/Engine/Plugins/Online/Experimental/OnlineServicesOSSAdapter/Source/Public/Online/AuthOSSAdapter.h b/Engine/Plugins/Online/Experimental/OnlineServicesOSSAdapter/Source/Public/Online/AuthOSSAdapter.h
new file mode 100644
--- /dev/null
+++ b/Engine/Plugins/Online/Experimental/OnlineServicesOSSAdapter/Source/Public/Online/AuthOSSAdapter.h
@@ -0,0 +1,41 @@
+// Copyright Epic Games, Inc. All Rights Reserved.
+
+#pragma once
+
+#include "Online/AuthCommon.h"
+
+#include "OnlineSubsystemTypes.h"
+
+class IOnlineIdentity;
+using IOnlineIdentityPtr = TSharedPtr<IOnlineIdentity>;
+
+namespace UE::Online {
+
+class FAuthOSSAdapter : public FAuthCommon
+{
+public:
+	using Super = FAuthCommon;
+
+	using FAuthCommon::FAuthCommon;
+
+	// IOnlineComponent
+	virtual void PostInitialize() override;
+	virtual void PreShutdown() override;
+
+	// IAuth
+	virtual TOnlineAsyncOpHandle<FAuthLogin> Login(FAuthLogin::Params&& Params) override;
+	virtual TOnlineAsyncOpHandle<FAuthLogout> Logout(FAuthLogout::Params&& Params) override;
+	virtual TOnlineAsyncOpHandle<FAuthGenerateAuth> GenerateAuth(FAuthGenerateAuth::Params&& Params) override;
+	virtual TOnlineResult<FAuthGetAccountByPlatformUserId> GetAccountByPlatformUserId(FAuthGetAccountByPlatformUserId::Params&& Params) override;
+	virtual TOnlineResult<FAuthGetAccountByAccountId> GetAccountByAccountId(FAuthGetAccountByAccountId::Params&& Params) override;
+
+	FUniqueNetIdRef GetUniqueNetId(FOnlineAccountIdHandle AccountIdHandle) const;
+	int32 GetLocalUserNum(FOnlineAccountIdHandle AccountIdHandle) const;
+
+	IOnlineIdentityPtr GetIdentityInterface() const;
+
+protected:
+	FDelegateHandle OnLoginStatusChangedHandle[MAX_LOCAL_PLAYERS];
+};
+
+/* UE::Online */ }
diff --git a/Engine/Plugins/Online/Experimental/OnlineServicesOSSAdapter/Source/Public/Online/OnlineIdOSSAdapter.h b/Engine/Plugins/Online/Experimental/OnlineServicesOSSAdapter/Source/Public/Online/OnlineIdOSSAdapter.h
new file mode 100644
--- /dev/null
+++ b/Engine/Plugins/Online/Experimental/OnlineServicesOSSAdapter/Source/Public/Online/OnlineIdOSSAdapter.h
@@ -0,0 +1,106 @@
+// Copyright Epic Games, Inc. All Rights Reserved.
+
+#pragma once
+
+#include "Online/CoreOnline.h"
+
+//#include "Online/OnlineServicesTypes.h"
+//#include "Online/OnlineId.h"
+
+#include "Misc/ScopeRWLock.h"
+
+namespace UE::Online {
+
+/**
+ * A net id registry suitable for use with OSS FUniqueNetIds
+ */
+class FOnlineUniqueNetIdRegistry : public IOnlineAccountIdRegistry
+{
+public:
+	FOnlineUniqueNetIdRegistry(EOnlineServices InOnlineServicesType)
+		: OnlineServicesType(InOnlineServicesType)
+	{
+	}
+
+	virtual ~FOnlineUniqueNetIdRegistry() {}
+
+	FOnlineAccountIdHandle FindOrAddHandle(const FUniqueNetIdRef& IdValue)
+	{
+		FOnlineAccountIdHandle Handle;
+		// Take a read lock and check if we already have a handle
+		{
+			FReadScopeLock ReadLock(Lock);
+			if (const uint32* FoundHandle = IdValueToHandleMap.Find(IdValue))
+			{
+				Handle = FOnlineAccountIdHandle(OnlineServicesType, *FoundHandle);
+			}
+		}
+
+		if (!Handle.IsValid())
+		{
+			// Take a write lock, check again if we already have a handle, or insert a new element.
+			FWriteScopeLock WriteLock(Lock);
+			if (const uint32* FoundHandle = IdValueToHandleMap.Find(IdValue))
+			{
+				Handle = FOnlineAccountIdHandle(OnlineServicesType, *FoundHandle);
+			}
+
+			if (!Handle.IsValid())
+			{
+				IdValues.Emplace(IdValue);
+				Handle = FOnlineAccountIdHandle(OnlineServicesType, IdValues.Num());
+				IdValueToHandleMap.Emplace(IdValue, Handle.GetHandle());
+			}
+		}
+
+		return Handle;
+	}
+
+	// Returns a copy as it's not thread safe to return a pointer/ref to an element of an array that can be relocated by another thread.
+	FUniqueNetIdRef GetIdValue(const FOnlineAccountIdHandle Handle) const
+	{
+		if (Handle.GetOnlineServicesType() == OnlineServicesType
+			&& Handle.IsValid()
+			&& IdValues.IsValidIndex(Handle.GetHandle() - 1))
+		{
+			FReadScopeLock ReadLock(Lock);
+			return IdValues[Handle.GetHandle() - 1];
+		}
+		return FUniqueNetIdString::EmptyId();
+	}
+
+	FUniqueNetIdRef GetIdValueChecked(const FOnlineAccountIdHandle Handle) const
+	{
+		check(Handle.GetOnlineServicesType() == OnlineServicesType
+			&& Handle.IsValid()
+			&& IdValues.IsValidIndex(Handle.GetHandle() - 1));
+		return GetIdValue(Handle);
+	}
+
+	// Begin IOnlineAccountIdRegistry
+	virtual FString ToLogString(const FOnlineAccountIdHandle& Handle) const override
+	{
+		return GetIdValue(Handle)->ToDebugString();
+	}
+
+	virtual TArray<uint8> ToReplicationData(const FOnlineAccountIdHandle& Handle) const
+	{
+		return TArray<uint8>();
+	}
+
+	virtual FOnlineAccountIdHandle FromReplicationData(const TArray<uint8>& Handle)
+	{
+		return FOnlineAccountIdHandle();
+	}
+
+	// End IOnlineAccountIdRegistry
+
+private:
+	mutable FRWLock Lock;
+
+	EOnlineServices OnlineServicesType;
+	TArray<FUniqueNetIdRef> IdValues;
+	TUniqueNetIdMap<uint32> IdValueToHandleMap;
+};
+
+/* UE::Online */ }
diff --git a/Engine/Plugins/Online/Experimental/OnlineServicesOSSAdapter/Source/Public/Online/OnlineServicesOSSAdapter.h b/Engine/Plugins/Online/Experimental/OnlineServicesOSSAdapter/Source/Public/Online/OnlineServicesOSSAdapter.h
new file mode 100644
--- /dev/null
+++ b/Engine/Plugins/Online/Experimental/OnlineServicesOSSAdapter/Source/Public/Online/OnlineServicesOSSAdapter.h
@@ -0,0 +1,32 @@
+// Copyright Epic Games, Inc. All Rights Reserved.
+
+#pragma once
+
+#include "Online/OnlineServicesCommon.h"
+
+class IOnlineSubsystem;
+
+namespace UE::Online {
+
+class FOnlineUniqueNetIdRegistry;
+
+class FOnlineServicesOSSAdapter : public FOnlineServicesCommon
+{
+public:
+	using Super = FOnlineServicesCommon;
+
+	ONLINESERVICESOSSADAPTER_API FOnlineServicesOSSAdapter(EOnlineServices InServicesType, const FString& InConfigName, IOnlineSubsystem* InSubsystem);
+
+	ONLINESERVICESOSSADAPTER_API virtual void RegisterComponents() override;
+	ONLINESERVICESOSSADAPTER_API virtual void Initialize() override;
+
+	IOnlineSubsystem& GetSubsystem() { return *Subsystem; }
+	FOnlineUniqueNetIdRegistry& GetAccountIdRegistry() { return *AccountIdRegistry; }
+
+protected:
+	EOnlineServices ServicesType;
+	IOnlineSubsystem* Subsystem;
+	FOnlineUniqueNetIdRegistry* AccountIdRegistry;
+};
+
+/* UE::Online */ }
