Adding missing includes

diff --git a/Engine/Plugins/Online/Experimental/OnlineServicesOSSAdapter/Source/Public/Online/OnlineIdOSSAdapter.h b/Engine/Plugins/Online/Experimental/OnlineServicesOSSAdapter/Source/Public/Online/OnlineIdOSSAdapter.h
--- a/Engine/Plugins/Online/Experimental/OnlineServicesOSSAdapter/Source/Public/Online/OnlineIdOSSAdapter.h
+++ b/Engine/Plugins/Online/Experimental/OnlineServicesOSSAdapter/Source/Public/Online/OnlineIdOSSAdapter.h
@@ -3,7 +3,7 @@
 #pragma once
 
 #include "Online/CoreOnline.h"
-
+#include "OnlineSubsystemTypes.h"
 //#include "Online/OnlineServicesTypes.h"
 //#include "Online/OnlineId.h"
 
diff --git a/Engine/Plugins/Tests/EditorTests/Source/EditorTests/Private/UnrealEd/ObjectToolsTests.cpp b/Engine/Plugins/Tests/EditorTests/Source/EditorTests/Private/UnrealEd/ObjectToolsTests.cpp
--- a/Engine/Plugins/Tests/EditorTests/Source/EditorTests/Private/UnrealEd/ObjectToolsTests.cpp
+++ b/Engine/Plugins/Tests/EditorTests/Source/EditorTests/Private/UnrealEd/ObjectToolsTests.cpp
@@ -4,6 +4,7 @@
 
 #include "Misc/Paths.h"
 #include "Misc/AutomationTest.h"
+#include "ObjectTools.h"
 #include "ScopedTransaction.h"
 #include "Tests/AutomationTestSettings.h"
 #include "Tests/AutomationEditorCommon.h"
