OnlineServicesAdapter
- configure available services/subsystems via config, for example, to enable Steam, enable the steam OSS Plugin and add the following to Engine.ini: [OnlineServices.OSSAdapter] +Services=(Service="Steam",ConfigName="OSSAdapter.Steam",OnlineSubsystem="Steam")
- Privileges, Connectivitiy, Auth interfaces

diff --git a/Engine/Plugins/Online/Experimental/OnlineServicesOSSAdapter/Source/Private/Online/AuthOSSAdapter.cpp b/Engine/Plugins/Online/Experimental/OnlineServicesOSSAdapter/Source/Private/Online/AuthOSSAdapter.cpp
--- a/Engine/Plugins/Online/Experimental/OnlineServicesOSSAdapter/Source/Private/Online/AuthOSSAdapter.cpp
+++ b/Engine/Plugins/Online/Experimental/OnlineServicesOSSAdapter/Source/Private/Online/AuthOSSAdapter.cpp
@@ -211,19 +211,123 @@ TOnlineAsyncOpHandle<FAuthLogout> FAuthOSSAdapter::Logout(FAuthLogout::Params&&
 
 TOnlineAsyncOpHandle<FAuthGenerateAuth> FAuthOSSAdapter::GenerateAuth(FAuthGenerateAuth::Params&& Params)
 {
-	// TODO: If no cached auth token available, call Login with appropriate parameters to generate an auth token
+	TSharedRef<TOnlineAsyncOp<FAuthGenerateAuth>> Op = GetOp<FAuthGenerateAuth>(MoveTemp(Params));
 
-	return Super::GenerateAuth(MoveTemp(Params));
+	if (!Op->IsReady())
+	{
+		FUniqueNetIdRef UniqueNetId = GetUniqueNetId(Params.LocalUserId);
+
+		if (!UniqueNetId->IsValid())
+		{
+			Op->SetError(Errors::InvalidParams());
+			return Op->GetHandle();
+		}
+
+		if (GetIdentityInterface()->GetLoginStatus(*UniqueNetId) != ::ELoginStatus::LoggedIn)
+		{
+			Op->SetError(Errors::InvalidAuth());
+		}
+		else
+		{
+			Op->SetResult(FAuthGenerateAuth::Result());
+		}
+	}
+
+	return Op->GetHandle();
+}
+
+TOnlineResult<FAuthGetAuthToken> FAuthOSSAdapter::GetAuthToken(FAuthGetAuthToken::Params&& Params)
+{
+	FUniqueNetIdRef UniqueNetId = GetUniqueNetId(Params.LocalUserId);
+
+	if (!UniqueNetId->IsValid())
+	{
+		return TOnlineResult<FAuthGetAuthToken>(Errors::InvalidUser());
+	}
+
+	if (GetIdentityInterface()->GetLoginStatus(*UniqueNetId) != ::ELoginStatus::LoggedIn)
+	{
+		return TOnlineResult<FAuthGetAuthToken>(Errors::InvalidAuth());
+	}
+
+	TSharedPtr<FUserOnlineAccount> UserAccount = GetIdentityInterface()->GetUserAccount(*UniqueNetId);
+	if (UserAccount.IsValid())
+	{
+		FAuthGetAuthToken::Result Result;
+		Result.Token = UserAccount->GetAccessToken();
+
+		return TOnlineResult<FAuthGetAuthToken>(Result);
+	}
+	else
+	{
+		// Fall back to GetAuthToken if the user account is not available (Steam doesn't implement GetUserAccount)
+		int32 LocalUserIndex = GetLocalUserNum(Params.LocalUserId);
+
+		if (LocalUserIndex < 0 || LocalUserIndex >= MAX_LOCAL_PLAYERS)
+		{
+			return TOnlineResult<FAuthGetAuthToken>(Errors::InvalidUser());
+		}
+
+		FString Token = GetIdentityInterface()->GetAuthToken(LocalUserIndex);
+		if (!Token.IsEmpty())
+		{
+			FAuthGetAuthToken::Result Result;
+			Result.Token = Token;
+
+			return TOnlineResult<FAuthGetAuthToken>(Result);
+		}
+		else
+		{
+			return TOnlineResult<FAuthGetAuthToken>(Errors::InvalidUser());
+		}
+	}
 }
 
 TOnlineResult<FAuthGetAccountByPlatformUserId> FAuthOSSAdapter::GetAccountByPlatformUserId(FAuthGetAccountByPlatformUserId::Params&& Params)
 {
-	return Super::GetAccountByPlatformUserId(MoveTemp(Params));
+	int32 LocalUserIndex = FPlatformMisc::GetUserIndexForPlatformUser(Params.PlatformUserId);
+
+	if (LocalUserIndex < 0 || LocalUserIndex >= MAX_LOCAL_PLAYERS)
+	{
+		return TOnlineResult<FAuthGetAccountByPlatformUserId>(Errors::InvalidUser());
+	}
+
+	FUniqueNetIdPtr UniqueNetId = GetIdentityInterface()->GetUniquePlayerId(LocalUserIndex);
+
+	if (!UniqueNetId.IsValid())
+	{
+		return TOnlineResult<FAuthGetAccountByPlatformUserId>(Errors::InvalidUser());
+	}
+
+	FOnlineAccountIdHandle Handle = static_cast<FOnlineServicesOSSAdapter&>(Services).GetAccountIdRegistry().FindOrAddHandle(UniqueNetId.ToSharedRef());
+	FAuthGetAccountByPlatformUserId::Result Result = { MakeShared<FAccountInfo>() };
+	Result.AccountInfo->PlatformUserId = Params.PlatformUserId;
+	Result.AccountInfo->UserId = Handle;
+	// v1 and v2 login status have the same values, so we can just cast here
+	Result.AccountInfo->LoginStatus = static_cast<ELoginStatus>(GetIdentityInterface()->GetLoginStatus(*UniqueNetId));
+	Result.AccountInfo->DisplayName = GetIdentityInterface()->GetPlayerNickname(*UniqueNetId);
+
+	return TOnlineResult<FAuthGetAccountByPlatformUserId>(MoveTemp(Result));
 }
 
 TOnlineResult<FAuthGetAccountByAccountId> FAuthOSSAdapter::GetAccountByAccountId(FAuthGetAccountByAccountId::Params&& Params)
 {
-	return Super::GetAccountByAccountId(MoveTemp(Params));
+	FUniqueNetIdRef UniqueNetId = GetUniqueNetId(Params.LocalUserId);
+
+	if (!UniqueNetId->IsValid())
+	{
+		return TOnlineResult<FAuthGetAccountByAccountId>(Errors::InvalidUser());
+	}
+
+	FOnlineAccountIdHandle Handle = static_cast<FOnlineServicesOSSAdapter&>(Services).GetAccountIdRegistry().FindOrAddHandle(UniqueNetId);
+	FAuthGetAccountByAccountId::Result Result = { MakeShared<FAccountInfo>() };
+	Result.AccountInfo->PlatformUserId = GetIdentityInterface()->GetPlatformUserIdFromUniqueNetId(*UniqueNetId);
+	Result.AccountInfo->UserId = Handle;
+	// v1 and v2 login status have the same values, so we can just cast here
+	Result.AccountInfo->LoginStatus = static_cast<ELoginStatus>(GetIdentityInterface()->GetLoginStatus(*UniqueNetId));
+	Result.AccountInfo->DisplayName = GetIdentityInterface()->GetPlayerNickname(*UniqueNetId);
+
+	return TOnlineResult<FAuthGetAccountByAccountId>(MoveTemp(Result));
 }
 
 FUniqueNetIdRef FAuthOSSAdapter::GetUniqueNetId(FOnlineAccountIdHandle AccountIdHandle) const
diff --git a/Engine/Plugins/Online/Experimental/OnlineServicesOSSAdapter/Source/Private/Online/ConnectivityOSSAdapter.cpp b/Engine/Plugins/Online/Experimental/OnlineServicesOSSAdapter/Source/Private/Online/ConnectivityOSSAdapter.cpp
new file mode 100644
--- /dev/null
+++ b/Engine/Plugins/Online/Experimental/OnlineServicesOSSAdapter/Source/Private/Online/ConnectivityOSSAdapter.cpp
@@ -0,0 +1,94 @@
+// Copyright Epic Games, Inc. All Rights Reserved.
+
+#include "Online/ConnectivityOSSAdapter.h"
+
+#include "Online/OnlineServicesOSSAdapter.h"
+
+#include "OnlineSubsystem.h"
+
+namespace UE::Online {
+
+void FConnectivityOSSAdapter::PostInitialize()
+{
+	Super::PostInitialize();
+
+	IOnlineSubsystem& Subsystem = static_cast<FOnlineServicesOSSAdapter&>(Services).GetSubsystem();
+	OnConnectionStatusChangedHandle = Subsystem.OnConnectionStatusChangedDelegates.AddLambda(
+		[WeakThis = TWeakPtr<IConnectivity>(AsShared()), this](const FString& ServiceName, EOnlineServerConnectionStatus::Type LastConnectionStatus, EOnlineServerConnectionStatus::Type NewConnectionStatus)
+		{
+			TSharedPtr<IConnectivity> PinnedThis = WeakThis.Pin();
+			if (PinnedThis.IsValid())
+			{
+				FConnectionStatusChanged Params;
+				Params.ServiceName = ServiceName;
+
+				switch (LastConnectionStatus)
+				{
+					case EOnlineServerConnectionStatus::Normal:
+					case EOnlineServerConnectionStatus::Connected:
+						Params.PreviousStatus = EOnlineServicesConnectionStatus::Connected;
+						break;
+
+					case EOnlineServerConnectionStatus::NotConnected:
+					case EOnlineServerConnectionStatus::ConnectionDropped:
+					case EOnlineServerConnectionStatus::NoNetworkConnection:
+					case EOnlineServerConnectionStatus::ServiceUnavailable:
+					case EOnlineServerConnectionStatus::UpdateRequired:
+					case EOnlineServerConnectionStatus::ServersTooBusy:
+					case EOnlineServerConnectionStatus::DuplicateLoginDetected:
+					case EOnlineServerConnectionStatus::InvalidUser:
+					case EOnlineServerConnectionStatus::NotAuthorized:
+					case EOnlineServerConnectionStatus::InvalidSession:
+						Params.PreviousStatus = EOnlineServicesConnectionStatus::NotConnected;
+						break;
+				}
+
+				switch (NewConnectionStatus)
+				{
+					case EOnlineServerConnectionStatus::Normal:
+					case EOnlineServerConnectionStatus::Connected:
+						Params.CurrentStatus = EOnlineServicesConnectionStatus::Connected;
+						break;
+
+					case EOnlineServerConnectionStatus::NotConnected:
+					case EOnlineServerConnectionStatus::ConnectionDropped:
+					case EOnlineServerConnectionStatus::NoNetworkConnection:
+					case EOnlineServerConnectionStatus::ServiceUnavailable:
+					case EOnlineServerConnectionStatus::UpdateRequired:
+					case EOnlineServerConnectionStatus::ServersTooBusy:
+					case EOnlineServerConnectionStatus::DuplicateLoginDetected:
+					case EOnlineServerConnectionStatus::InvalidUser:
+					case EOnlineServerConnectionStatus::NotAuthorized:
+					case EOnlineServerConnectionStatus::InvalidSession:
+						Params.CurrentStatus = EOnlineServicesConnectionStatus::NotConnected;
+						break;
+				}
+
+				CurrentStatus.Add(ServiceName, Params.CurrentStatus);
+
+				OnConnectionStatusChangedEvent.Broadcast(Params);
+			}
+		});
+}
+
+void FConnectivityOSSAdapter::PreShutdown()
+{
+	IOnlineSubsystem& Subsystem = static_cast<FOnlineServicesOSSAdapter&>(Services).GetSubsystem();
+	if (OnConnectionStatusChangedHandle.IsValid())
+	{
+		Subsystem.OnConnectionStatusChangedDelegates.Remove(OnConnectionStatusChangedHandle);
+		OnConnectionStatusChangedHandle.Reset();
+	}
+
+	Super::PreShutdown();
+}
+
+TOnlineResult<FGetConnectionStatus> FConnectivityOSSAdapter::GetConnectionStatus(FGetConnectionStatus::Params&& Params)
+{
+	FGetConnectionStatus::Result Result;
+	Result.Status = CurrentStatus.FindOrAdd(Params.ServiceName, EOnlineServicesConnectionStatus::NotConnected);
+
+	return TOnlineResult<FGetConnectionStatus>(Result);
+}
+
+/* UE::Online*/ }
diff --git a/Engine/Plugins/Online/Experimental/OnlineServicesOSSAdapter/Source/Private/Online/ExternalUIOSSAdapter.cpp b/Engine/Plugins/Online/Experimental/OnlineServicesOSSAdapter/Source/Private/Online/ExternalUIOSSAdapter.cpp
new file mode 100644
--- /dev/null
+++ b/Engine/Plugins/Online/Experimental/OnlineServicesOSSAdapter/Source/Private/Online/ExternalUIOSSAdapter.cpp
@@ -0,0 +1,136 @@
+// Copyright Epic Games, Inc. All Rights Reserved.
+
+#include "Online/ExternalUIOSSAdapter.h"
+
+#include "Online/OnlineServicesOSSAdapter.h"
+#include "Online/AuthOSSAdapter.h"
+#include "Online/OnlineIdOSSAdapter.h"
+
+#include "OnlineSubsystem.h"
+#include "Interfaces/OnlineExternalUIInterface.h"
+
+namespace UE::Online {
+
+void FExternalUIOSSAdapter::PostInitialize()
+{
+	Super::PostInitialize();
+
+	OnExternalUIChangeHandle = GetExternalUIInterface()->OnExternalUIChangeDelegates.AddLambda(
+		[WeakThis = TWeakPtr<IExternalUI>(AsShared()), this](bool bIsOpening)
+		{
+			TSharedPtr<IExternalUI> PinnedThis = WeakThis.Pin();
+			if (PinnedThis.IsValid())
+			{
+				FExternalUIStatusChanged StatusChanged;
+				StatusChanged.bIsOpening = bIsOpening;
+
+				OnExternalUIStatusChangedEvent.Broadcast(StatusChanged);
+			}
+		});
+}
+
+void FExternalUIOSSAdapter::PreShutdown()
+{
+	if (OnExternalUIChangeHandle.IsValid())
+	{
+		GetExternalUIInterface()->OnExternalUIChangeDelegates.Remove(OnExternalUIChangeHandle);
+		OnExternalUIChangeHandle.Reset();
+	}
+
+	Super::PreShutdown();
+}
+
+TOnlineAsyncOpHandle<FExternalUIShowLoginUI> FExternalUIOSSAdapter::ShowLoginUI(FExternalUIShowLoginUI::Params&& Params)
+{
+	TSharedRef<TOnlineAsyncOp<FExternalUIShowLoginUI>> Op = GetOp<FExternalUIShowLoginUI>(MoveTemp(Params));
+
+	if (!Op->IsReady())
+	{
+		int32 LocalUserIndex = Op->GetParams().PlatformUserId.GetInternalId();
+
+		if (LocalUserIndex < 0 || LocalUserIndex >= MAX_LOCAL_PLAYERS)
+		{
+			Op->SetError(Errors::InvalidParams());
+			return Op->GetHandle();
+		}
+
+		Op->Then([this, LocalUserIndex](TOnlineAsyncOp<FExternalUIShowLoginUI>& Op)
+			{
+				TSharedPtr<TPromise<void>> Promise = MakeShared<TPromise<void>>();
+				bool bShowedUI = GetExternalUIInterface()->ShowLoginUI(LocalUserIndex, true, false,
+					FOnLoginUIClosedDelegate::CreateLambda(
+						[this, WeakOp = Op.AsWeak(), Promise](FUniqueNetIdPtr UniqueId, const int ControllerIndex, const ::FOnlineError& Error)
+						{
+							TSharedPtr<TOnlineAsyncOp<FExternalUIShowLoginUI>> PinnedOp = WeakOp.Pin();
+							if (PinnedOp.IsValid())
+							{
+								if (Error.WasSuccessful())
+								{
+									FOnlineAccountIdHandle Handle = static_cast<FOnlineServicesOSSAdapter&>(Services).GetAccountIdRegistry().FindOrAddHandle(UniqueId.ToSharedRef());
+									FExternalUIShowLoginUI::Result Result = { MakeShared<FAccountInfo>() };
+									Result.AccountInfo->PlatformUserId = PinnedOp->GetParams().PlatformUserId;
+									Result.AccountInfo->UserId = Handle;
+									Result.AccountInfo->LoginStatus = ELoginStatus::LoggedIn;
+
+									PinnedOp->SetResult(MoveTemp(Result));
+								}
+								else
+								{
+									PinnedOp->SetError(Errors::Unknown()); // TODO: Translate v1 to v2 error
+								}
+								
+								Promise->SetValue();
+							}
+						}
+					));
+
+				if (!bShowedUI)
+				{
+					Op.SetError(Errors::Unknown());
+				}
+
+				return Promise->GetFuture();
+			})
+			.Enqueue(GetSerialQueue());
+	}
+
+	return Op->GetHandle();
+}
+
+TOnlineAsyncOpHandle<FExternalUIShowFriendsUI> FExternalUIOSSAdapter::ShowFriendsUI(FExternalUIShowFriendsUI::Params&& Params)
+{
+	TSharedRef<TOnlineAsyncOp<FExternalUIShowFriendsUI>> Op = GetOp<FExternalUIShowFriendsUI>(MoveTemp(Params));
+
+	if (!Op->IsReady())
+	{
+		int32 LocalUserIndex = Services.Get<FAuthOSSAdapter>()->GetLocalUserNum(Op->GetParams().LocalUserId);
+
+		if (LocalUserIndex < 0 || LocalUserIndex >= MAX_LOCAL_PLAYERS)
+		{
+			Op->SetError(Errors::InvalidParams());
+			return Op->GetHandle();
+		}
+
+		Op->Then([this, LocalUserIndex](TOnlineAsyncOp<FExternalUIShowFriendsUI>& Op)
+			{
+				if (GetExternalUIInterface()->ShowFriendsUI(LocalUserIndex))
+				{
+					Op.SetResult(FExternalUIShowFriendsUI::Result());
+				}
+				else
+				{
+					Op.SetError(Errors::Unknown());
+				}
+			})
+			.Enqueue(GetSerialQueue());
+	}
+
+	return Op->GetHandle();
+}
+
+IOnlineExternalUIPtr FExternalUIOSSAdapter::GetExternalUIInterface() const
+{
+	return static_cast<FOnlineServicesOSSAdapter&>(Services).GetSubsystem().GetExternalUIInterface();
+}
+
+/* UE::Online */ }
diff --git a/Engine/Plugins/Online/Experimental/OnlineServicesOSSAdapter/Source/Private/Online/OnlineServicesOSSAdapter.cpp b/Engine/Plugins/Online/Experimental/OnlineServicesOSSAdapter/Source/Private/Online/OnlineServicesOSSAdapter.cpp
--- a/Engine/Plugins/Online/Experimental/OnlineServicesOSSAdapter/Source/Private/Online/OnlineServicesOSSAdapter.cpp
+++ b/Engine/Plugins/Online/Experimental/OnlineServicesOSSAdapter/Source/Private/Online/OnlineServicesOSSAdapter.cpp
@@ -4,6 +4,9 @@
 
 #include "Online/OnlineIdOSSAdapter.h"
 #include "Online/AuthOSSAdapter.h"
+#include "Online/ConnectivityOSSAdapter.h"
+#include "Online/PrivilegesOSSAdapter.h"
+#include "Online/ExternalUIOSSAdapter.h"
 
 #include "OnlineSubsystem.h"
 
@@ -19,6 +22,13 @@ FOnlineServicesOSSAdapter::FOnlineServicesOSSAdapter(EOnlineServices InServicesT
 void FOnlineServicesOSSAdapter::RegisterComponents()
 {
 	Components.Register<FAuthOSSAdapter>(*this);
+	Components.Register<FConnectivityOSSAdapter>(*this);
+	Components.Register<FPrivilegesOSSAdapter>(*this);
+
+	if (Subsystem->GetExternalUIInterface().IsValid())
+	{
+		Components.Register<FExternalUIOSSAdapter>(*this);
+	}
 
 	FOnlineServicesCommon::RegisterComponents();
 }
diff --git a/Engine/Plugins/Online/Experimental/OnlineServicesOSSAdapter/Source/Private/Online/OnlineServicesOSSAdapterModule.cpp b/Engine/Plugins/Online/Experimental/OnlineServicesOSSAdapter/Source/Private/Online/OnlineServicesOSSAdapterModule.cpp
--- a/Engine/Plugins/Online/Experimental/OnlineServicesOSSAdapter/Source/Private/Online/OnlineServicesOSSAdapterModule.cpp
+++ b/Engine/Plugins/Online/Experimental/OnlineServicesOSSAdapter/Source/Private/Online/OnlineServicesOSSAdapterModule.cpp
@@ -4,7 +4,11 @@
 #include "Modules/ModuleInterface.h"
 #include "Modules/ModuleManager.h"
 
+#include "Online/OnlineServicesOSSAdapter.h"
 #include "Online/OnlineServicesRegistry.h"
+#include "Online/OnlineIdOSSAdapter.h"
+
+#include "OnlineSubsystem.h"
 
 namespace UE::Online
 {
@@ -13,11 +17,149 @@ class FOnlineServicesOSSAdapterModule : public IModuleInterface
 {
 public:
 	virtual void StartupModule() override;
+	virtual void ShutdownModule() override;
 protected:
 };
 
+struct FOSSAdapterService
+{
+	EOnlineServices Service = EOnlineServices::Default;
+	FString ConfigName;
+	FName OnlineSubsystem;
+	int Priority = -1;
+};
+
+struct FOSSAdapterConfig
+{
+	// TArray<FOSSAdapterService> Services; // use once online config parsing supports arrays of structs
+	TArray<FString> Services;
+};
+
+namespace Meta {
+
+BEGIN_ONLINE_STRUCT_META(FOSSAdapterService)
+	ONLINE_STRUCT_FIELD(FOSSAdapterService, Service),
+	ONLINE_STRUCT_FIELD(FOSSAdapterService, ConfigName),
+	ONLINE_STRUCT_FIELD(FOSSAdapterService, OnlineSubsystem),
+	ONLINE_STRUCT_FIELD(FOSSAdapterService, Priority)
+END_ONLINE_STRUCT_META()
+
+BEGIN_ONLINE_STRUCT_META(FOSSAdapterConfig)
+	ONLINE_STRUCT_FIELD(FOSSAdapterConfig, Services)
+END_ONLINE_STRUCT_META()
+
+/* Meta */ }
+
+class FOnlineServicesFactoryOSSAdapter : public IOnlineServicesFactory
+{
+public:
+	FOnlineServicesFactoryOSSAdapter(const FOSSAdapterService& InConfig)
+		: Config(InConfig)
+	{
+	}
+
+	virtual ~FOnlineServicesFactoryOSSAdapter() {}
+	virtual TSharedPtr<IOnlineServices> Create() override
+	{
+		IOnlineSubsystem* Subsystem = IOnlineSubsystem::Get(Config.OnlineSubsystem);
+		if (Subsystem != nullptr)
+		{
+			return MakeShared<FOnlineServicesOSSAdapter>(Config.Service, Config.ConfigName, Subsystem);
+		}
+		else
+		{
+			return nullptr;
+		}
+	}
+protected:
+	FOSSAdapterService Config;
+};
+
+namespace {
+
+// temporary parsing until online config system has full support for parsing arrays of structs/nested structs
+FString StripQuotes(const FString& Source)
+{
+	if (Source.StartsWith(TEXT("\"")))
+	{
+		return Source.Mid(1, Source.Len() - 2);
+	}
+	return Source;
+}
+
+void ParseServiceConfig(FOSSAdapterService& ServiceConfig, const FString& ServiceConfigString)
+{
+	const TCHAR* Delims[4] = { TEXT("("), TEXT(")"), TEXT("="), TEXT(",") };
+	TArray<FString> Values;
+	ServiceConfigString.ParseIntoArray(Values, Delims, 4, false);
+	for (int32 ValueIndex = 0; ValueIndex + 1 < Values.Num(); ++ValueIndex)
+	{
+		if (Values[ValueIndex].IsEmpty())
+		{
+			continue;
+		}
+
+		if (Values[ValueIndex] == TEXT("Service"))
+		{
+			LexFromString(ServiceConfig.Service, *StripQuotes(Values[ValueIndex + 1]));
+		}
+		else if (Values[ValueIndex] == TEXT("ConfigName"))
+		{
+			ServiceConfig.ConfigName = StripQuotes(Values[ValueIndex + 1]);
+		}
+		else if (Values[ValueIndex] == TEXT("OnlineSubsystem"))
+		{
+			ServiceConfig.OnlineSubsystem = FName(StripQuotes(Values[ValueIndex + 1]));
+		}
+		else if (Values[ValueIndex] == TEXT("Priority"))
+		{
+			ServiceConfig.Priority = FCString::Atoi(*StripQuotes(Values[ValueIndex + 1]));
+		}
+
+		++ValueIndex;
+	}
+}
+
+/* unnamed */ }
+
 void FOnlineServicesOSSAdapterModule::StartupModule()
 {
+	FOnlineConfigProviderGConfig ConfigProvider(GEngineIni);
+	FOSSAdapterConfig Config;
+	if (LoadConfig(ConfigProvider, TEXT("OnlineServices.OSSAdapter"), Config))
+	{
+		for (const FString& ServiceConfigString : Config.Services)
+		{
+			FOSSAdapterService ServiceConfig;
+			ParseServiceConfig(ServiceConfig, ServiceConfigString);
+			
+			if (IOnlineSubsystem::IsEnabled(ServiceConfig.OnlineSubsystem))
+			{
+				FOnlineServicesRegistry::Get().RegisterServicesFactory(ServiceConfig.Service, MakeUnique<FOnlineServicesFactoryOSSAdapter>(ServiceConfig), ServiceConfig.Priority);
+				FOnlineIdRegistryRegistry::Get().RegisterAccountIdRegistry(ServiceConfig.Service, new FOnlineUniqueNetIdRegistry(ServiceConfig.Service), ServiceConfig.Priority);
+			}
+		}
+	}
+}
+
+void FOnlineServicesOSSAdapterModule::ShutdownModule()
+{
+	FOnlineConfigProviderGConfig ConfigProvider(GEngineIni);
+	FOSSAdapterConfig Config;
+	if (LoadConfig(ConfigProvider, TEXT("OnlineServices.OSSAdapter"), Config))
+	{
+		for (const FString& ServiceConfigString : Config.Services)
+		{
+			FOSSAdapterService ServiceConfig;
+			ParseServiceConfig(ServiceConfig, ServiceConfigString);
+
+			if (IOnlineSubsystem::IsEnabled(ServiceConfig.OnlineSubsystem))
+			{
+				FOnlineServicesRegistry::Get().UnregisterServicesFactory(ServiceConfig.Service, ServiceConfig.Priority);
+				FOnlineIdRegistryRegistry::Get().UnregisterAccountIdRegistry(ServiceConfig.Service, ServiceConfig.Priority);
+			}
+		}
+	}
 }
 
 /* UE::Online */ }
diff --git a/Engine/Plugins/Online/Experimental/OnlineServicesOSSAdapter/Source/Private/Online/PrivilegesOSSAdapter.cpp b/Engine/Plugins/Online/Experimental/OnlineServicesOSSAdapter/Source/Private/Online/PrivilegesOSSAdapter.cpp
new file mode 100644
--- /dev/null
+++ b/Engine/Plugins/Online/Experimental/OnlineServicesOSSAdapter/Source/Private/Online/PrivilegesOSSAdapter.cpp
@@ -0,0 +1,81 @@
+// Copyright Epic Games, Inc. All Rights Reserved.
+
+#include "Online/PrivilegesOSSAdapter.h"
+
+#include "Online/OnlineServicesOSSAdapter.h"
+#include "Online/AuthOSSAdapter.h"
+
+#include "OnlineSubsystem.h"
+#include "Interfaces/OnlineIdentityInterface.h"
+
+namespace UE::Online {
+
+void FPrivilegesOSSAdapter::PostInitialize()
+{
+	Super::PostInitialize();
+}
+
+void FPrivilegesOSSAdapter::PreShutdown()
+{
+	Super::PreShutdown();
+}
+
+TOnlineAsyncOpHandle<FQueryUserPrivilege> FPrivilegesOSSAdapter::QueryUserPrivilege(FQueryUserPrivilege::Params&& Params)
+{
+	TSharedRef<TOnlineAsyncOp<FQueryUserPrivilege>> Op = GetOp<FQueryUserPrivilege>(MoveTemp(Params));
+
+	if (!Op->IsReady())
+	{
+		FAuthOSSAdapter* Auth = Services.Get<FAuthOSSAdapter>();
+		FUniqueNetIdRef UniqueNetId = Auth->GetUniqueNetId(Op->GetParams().LocalUserId);
+		if (!UniqueNetId->IsValid())
+		{
+			Op->SetError(Errors::InvalidUser());
+			return Op->GetHandle();
+		}
+
+		Op->Then([this, UniqueNetId](TOnlineAsyncOp<FQueryUserPrivilege>& Op)
+			{
+				TSharedPtr<TPromise<void>> Promise = MakeShared<TPromise<void>>();
+				::EUserPrivileges::Type Privilege;
+				switch (Op.GetParams().Privilege)
+				{
+					default:
+					case EUserPrivileges::CanPlay:						Privilege = ::EUserPrivileges::CanPlay;						break;
+					case EUserPrivileges::CanPlayOnline:				Privilege = ::EUserPrivileges::CanPlayOnline;				break;
+					case EUserPrivileges::CanCommunicateViaTextOnline:	Privilege = ::EUserPrivileges::CanCommunicateOnline;		break;
+					case EUserPrivileges::CanCommunicateViaVoiceOnline:	Privilege = ::EUserPrivileges::CanCommunicateOnline;		break;
+					case EUserPrivileges::CanUseUserGeneratedContent:	Privilege = ::EUserPrivileges::CanUseUserGeneratedContent;	break;
+					case EUserPrivileges::CanCrossPlay:					Privilege = ::EUserPrivileges::CanUserCrossPlay;			break;
+				}
+
+				GetIdentityInterface()->GetUserPrivilege(*UniqueNetId, Privilege, 
+					IOnlineIdentity::FOnGetUserPrivilegeCompleteDelegate::CreateLambda(
+						[WeakOp = Op.AsWeak(), Promise](const FUniqueNetId& LocalUserId, ::EUserPrivileges::Type Privilege, uint32 PrivilegeResult)
+						{
+							TSharedPtr<TOnlineAsyncOp<FQueryUserPrivilege>> PinnedOp = WeakOp.Pin();
+							if (PinnedOp.IsValid())
+							{
+								FQueryUserPrivilege::Result Result;
+								// OSS EPrivilegeResults currently has the same values as OnlineServices EPrivilegeResults
+								Result.PrivilegeResult = static_cast<EPrivilegeResults>(PrivilegeResult);
+								PinnedOp->SetResult(MoveTemp(Result));
+
+								Promise->SetValue();
+							}
+						}));
+
+				return Promise->GetFuture();
+			})
+			.Enqueue(GetSerialQueue());
+	}
+
+	return Op->GetHandle();
+}
+
+IOnlineIdentityPtr FPrivilegesOSSAdapter::GetIdentityInterface() const
+{
+	return static_cast<FOnlineServicesOSSAdapter&>(Services).GetSubsystem().GetIdentityInterface();
+}
+
+/* UE::Online */ }
diff --git a/Engine/Plugins/Online/Experimental/OnlineServicesOSSAdapter/Source/Public/Online/AuthOSSAdapter.h b/Engine/Plugins/Online/Experimental/OnlineServicesOSSAdapter/Source/Public/Online/AuthOSSAdapter.h
--- a/Engine/Plugins/Online/Experimental/OnlineServicesOSSAdapter/Source/Public/Online/AuthOSSAdapter.h
+++ b/Engine/Plugins/Online/Experimental/OnlineServicesOSSAdapter/Source/Public/Online/AuthOSSAdapter.h
@@ -26,6 +26,7 @@ public:
 	virtual TOnlineAsyncOpHandle<FAuthLogin> Login(FAuthLogin::Params&& Params) override;
 	virtual TOnlineAsyncOpHandle<FAuthLogout> Logout(FAuthLogout::Params&& Params) override;
 	virtual TOnlineAsyncOpHandle<FAuthGenerateAuth> GenerateAuth(FAuthGenerateAuth::Params&& Params) override;
+	virtual TOnlineResult<FAuthGetAuthToken> GetAuthToken(FAuthGetAuthToken::Params&& Params) override;
 	virtual TOnlineResult<FAuthGetAccountByPlatformUserId> GetAccountByPlatformUserId(FAuthGetAccountByPlatformUserId::Params&& Params) override;
 	virtual TOnlineResult<FAuthGetAccountByAccountId> GetAccountByAccountId(FAuthGetAccountByAccountId::Params&& Params) override;
 
diff --git a/Engine/Plugins/Online/Experimental/OnlineServicesOSSAdapter/Source/Public/Online/ConnectivityOSSAdapter.h b/Engine/Plugins/Online/Experimental/OnlineServicesOSSAdapter/Source/Public/Online/ConnectivityOSSAdapter.h
new file mode 100644
--- /dev/null
+++ b/Engine/Plugins/Online/Experimental/OnlineServicesOSSAdapter/Source/Public/Online/ConnectivityOSSAdapter.h
@@ -0,0 +1,28 @@
+// Copyright Epic Games, Inc. All Rights Reserved.
+
+#pragma once
+
+#include "Online/ConnectivityCommon.h"
+
+namespace UE::Online {
+
+class FConnectivityOSSAdapter : public FConnectivityCommon
+{
+public:
+	using Super = FConnectivityCommon;
+
+	using FConnectivityCommon::FConnectivityCommon;
+
+	// IOnlineComponent
+	virtual void PostInitialize() override;
+	virtual void PreShutdown() override;
+
+	// IConnectivity
+	virtual TOnlineResult<FGetConnectionStatus> GetConnectionStatus(FGetConnectionStatus::Params&& Params) override;
+
+protected:
+	FDelegateHandle OnConnectionStatusChangedHandle;
+	TMap<FString, EOnlineServicesConnectionStatus> CurrentStatus;
+};
+
+/* UE::Online */ }
diff --git a/Engine/Plugins/Online/Experimental/OnlineServicesOSSAdapter/Source/Public/Online/ExternalUIOSSAdapter.h b/Engine/Plugins/Online/Experimental/OnlineServicesOSSAdapter/Source/Public/Online/ExternalUIOSSAdapter.h
new file mode 100644
--- /dev/null
+++ b/Engine/Plugins/Online/Experimental/OnlineServicesOSSAdapter/Source/Public/Online/ExternalUIOSSAdapter.h
@@ -0,0 +1,33 @@
+// Copyright Epic Games, Inc. All Rights Reserved.
+
+#pragma once
+
+#include "Online/ExternalUICommon.h"
+
+class IOnlineExternalUI;
+using IOnlineExternalUIPtr = TSharedPtr<IOnlineExternalUI>;
+
+namespace UE::Online {
+
+class FExternalUIOSSAdapter : public FExternalUICommon
+{
+public:
+	using Super = FExternalUICommon;
+
+	using FExternalUICommon::FExternalUICommon;
+
+	// TOnlineComponent
+	virtual void PostInitialize() override;
+	virtual void PreShutdown() override;
+
+	// IExternalUI
+	virtual TOnlineAsyncOpHandle<FExternalUIShowLoginUI> ShowLoginUI(FExternalUIShowLoginUI::Params&& Params) override;
+	virtual TOnlineAsyncOpHandle<FExternalUIShowFriendsUI> ShowFriendsUI(FExternalUIShowFriendsUI::Params&& Params) override;
+
+	IOnlineExternalUIPtr GetExternalUIInterface() const;
+
+protected:
+	FDelegateHandle OnExternalUIChangeHandle;
+};
+
+/* UE::Online */ }
diff --git a/Engine/Plugins/Online/Experimental/OnlineServicesOSSAdapter/Source/Public/Online/OnlineIdOSSAdapter.h b/Engine/Plugins/Online/Experimental/OnlineServicesOSSAdapter/Source/Public/Online/OnlineIdOSSAdapter.h
--- a/Engine/Plugins/Online/Experimental/OnlineServicesOSSAdapter/Source/Public/Online/OnlineIdOSSAdapter.h
+++ b/Engine/Plugins/Online/Experimental/OnlineServicesOSSAdapter/Source/Public/Online/OnlineIdOSSAdapter.h
@@ -4,8 +4,6 @@
 
 #include "Online/CoreOnline.h"
 #include "OnlineSubsystemTypes.h"
-//#include "Online/OnlineServicesTypes.h"
-//#include "Online/OnlineId.h"
 
 #include "Misc/ScopeRWLock.h"
 
diff --git a/Engine/Plugins/Online/Experimental/OnlineServicesOSSAdapter/Source/Public/Online/PrivilegesOSSAdapter.h b/Engine/Plugins/Online/Experimental/OnlineServicesOSSAdapter/Source/Public/Online/PrivilegesOSSAdapter.h
new file mode 100644
--- /dev/null
+++ b/Engine/Plugins/Online/Experimental/OnlineServicesOSSAdapter/Source/Public/Online/PrivilegesOSSAdapter.h
@@ -0,0 +1,30 @@
+// Copyright Epic Games, Inc. All Rights Reserved.
+
+#pragma once
+
+#include "Online/PrivilegesCommon.h"
+
+class IOnlineIdentity;
+using IOnlineIdentityPtr = TSharedPtr<IOnlineIdentity>;
+
+namespace UE::Online {
+
+class FPrivilegesOSSAdapter : public FPrivilegesCommon
+{
+public:
+	using Super = FPrivilegesCommon;
+
+	using FPrivilegesCommon::FPrivilegesCommon;
+
+	// TOnlineComponent
+	virtual void PostInitialize() override;
+	virtual void PreShutdown() override;
+
+	// IPermission
+	virtual TOnlineAsyncOpHandle<FQueryUserPrivilege> QueryUserPrivilege(FQueryUserPrivilege::Params&& Params) override;
+
+protected:
+	IOnlineIdentityPtr GetIdentityInterface() const;
+};
+
+/* UE::Online */ }
