- Add platform auth login to Lyra
- Move cache dir logic from OnlineSubsystemEOS to EOSShared, to share with OnlineServicesEOS.

diff --git a/Engine/Plugins/Online/EOSShared/Source/EOSShared/Private/EOSSDKManager.cpp b/Engine/Plugins/Online/EOSShared/Source/EOSShared/Private/EOSSDKManager.cpp
--- a/Engine/Plugins/Online/EOSShared/Source/EOSShared/Private/EOSSDKManager.cpp
+++ b/Engine/Plugins/Online/EOSShared/Source/EOSShared/Private/EOSSDKManager.cpp
@@ -340,6 +340,11 @@ FString FEOSSDKManager::GetProductVersion() const
 	return FApp::GetBuildVersion();
 }
 
+FString FEOSSDKManager::GetCacheDirBase() const
+{
+	return FPlatformProcess::UserDir();
+}
+
 void FEOSSDKManager::ReleasePlatform(EOS_HPlatform PlatformHandle)
 {
 	if (ensure(ActivePlatforms.Contains(PlatformHandle)
diff --git a/Engine/Plugins/Online/EOSShared/Source/EOSShared/Private/EOSSDKManager.h b/Engine/Plugins/Online/EOSShared/Source/EOSShared/Private/EOSSDKManager.h
--- a/Engine/Plugins/Online/EOSShared/Source/EOSShared/Private/EOSSDKManager.h
+++ b/Engine/Plugins/Online/EOSShared/Source/EOSShared/Private/EOSSDKManager.h
@@ -30,6 +30,7 @@ public:
 
 	virtual FString GetProductName() const override;
 	virtual FString GetProductVersion() const override;
+	virtual FString GetCacheDirBase() const override;
 	// End IEOSSDKManager
 
 	void Shutdown();
diff --git a/Engine/Plugins/Online/EOSShared/Source/EOSShared/Private/IOS/IOSEOSSDKManager.cpp b/Engine/Plugins/Online/EOSShared/Source/EOSShared/Private/IOS/IOSEOSSDKManager.cpp
new file mode 100644
--- /dev/null
+++ b/Engine/Plugins/Online/EOSShared/Source/EOSShared/Private/IOS/IOSEOSSDKManager.cpp
@@ -0,0 +1,17 @@
+// Copyright Epic Games, Inc. All Rights Reserved.
+
+#if WITH_EOS_SDK
+
+#include "IOSEOSSDKManager.h"
+
+FString FIOSEOSSDKManager::GetCacheDirBase() const
+{
+	NSString* BundleIdentifier = [[NSBundle mainBundle]bundleIdentifier];
+	NSString* CacheDirectory = NSTemporaryDirectory(); // Potentially use NSCachesDirectory
+	CacheDirectory = [CacheDirectory stringByAppendingPathComponent : BundleIdentifier];
+
+	const char* CStrCacheDirectory = [CacheDirectory UTF8String];
+	return FString(UTF8_TO_TCHAR(CStrCacheDirectory));
+};
+
+#endif // WITH_EOS_SDK
\ No newline at end of file
diff --git a/Engine/Plugins/Online/EOSShared/Source/EOSShared/Private/IOS/IOSEOSSDKManager.h b/Engine/Plugins/Online/EOSShared/Source/EOSShared/Private/IOS/IOSEOSSDKManager.h
--- a/Engine/Plugins/Online/EOSShared/Source/EOSShared/Private/IOS/IOSEOSSDKManager.h
+++ b/Engine/Plugins/Online/EOSShared/Source/EOSShared/Private/IOS/IOSEOSSDKManager.h
@@ -6,6 +6,11 @@
 
 #include "EOSSDKManager.h"
 
-using FPlatformEOSSDKManager = FEOSSDKManager;
+class FIOSEOSSDKManager : public FEOSSDKManager
+{
+	virtual FString GetCacheDirBase() const override;
+};
+
+using FPlatformEOSSDKManager = FIOSEOSSDKManager;
 
 #endif // WITH_EOS_SDK
\ No newline at end of file
diff --git a/Engine/Plugins/Online/EOSShared/Source/EOSShared/Public/IEOSSDKManager.h b/Engine/Plugins/Online/EOSShared/Source/EOSShared/Public/IEOSSDKManager.h
--- a/Engine/Plugins/Online/EOSShared/Source/EOSShared/Public/IEOSSDKManager.h
+++ b/Engine/Plugins/Online/EOSShared/Source/EOSShared/Public/IEOSSDKManager.h
@@ -57,6 +57,7 @@ public:
 
 	virtual FString GetProductName() const = 0;
 	virtual FString GetProductVersion() const = 0;
+	virtual FString GetCacheDirBase() const = 0;
 
 	FEOSSDKManagerOnPreInitializeSDK OnPreInitializeSDK;
 	FEOSSDKManagerOnPreCreatePlatform OnPreCreatePlatform;
diff --git a/Engine/Plugins/Online/Experimental/OnlineServices/Source/OnlineServicesCommon/Private/Online/AuthCommon.cpp b/Engine/Plugins/Online/Experimental/OnlineServices/Source/OnlineServicesCommon/Private/Online/AuthCommon.cpp
--- a/Engine/Plugins/Online/Experimental/OnlineServices/Source/OnlineServicesCommon/Private/Online/AuthCommon.cpp
+++ b/Engine/Plugins/Online/Experimental/OnlineServices/Source/OnlineServicesCommon/Private/Online/AuthCommon.cpp
@@ -17,8 +17,9 @@ void FAuthCommon::RegisterCommands()
 {
 	RegisterCommand(&FAuthCommon::Login);
 	RegisterCommand(&FAuthCommon::Logout);
-	RegisterCommand(&FAuthCommon::GenerateAuth);
+	RegisterCommand(&FAuthCommon::GenerateAuthToken);
 	RegisterCommand(&FAuthCommon::GetAuthToken);
+	RegisterCommand(&FAuthCommon::GenerateAuthCode);
 	RegisterCommand(&FAuthCommon::GetAccountByPlatformUserId);
 	RegisterCommand(&FAuthCommon::GetAccountByAccountId);
 }
@@ -37,9 +38,9 @@ TOnlineAsyncOpHandle<FAuthLogout> FAuthCommon::Logout(FAuthLogout::Params&& Para
 	return Operation->GetHandle();
 }
 
-TOnlineAsyncOpHandle<FAuthGenerateAuth> FAuthCommon::GenerateAuth(FAuthGenerateAuth::Params&& Params)
+TOnlineAsyncOpHandle<FAuthGenerateAuthToken> FAuthCommon::GenerateAuthToken(FAuthGenerateAuthToken::Params&& Params)
 {
-	TOnlineAsyncOpRef<FAuthGenerateAuth> Operation = GetOp<FAuthGenerateAuth>(MoveTemp(Params));
+	TOnlineAsyncOpRef<FAuthGenerateAuthToken> Operation = GetOp<FAuthGenerateAuthToken>(MoveTemp(Params));
 	Operation->SetError(Errors::NotImplemented());
 	return Operation->GetHandle();
 }
@@ -49,6 +50,13 @@ TOnlineResult<FAuthGetAuthToken> FAuthCommon::GetAuthToken(FAuthGetAuthToken::Pa
 	return TOnlineResult<FAuthGetAuthToken>(Errors::NotImplemented());
 }
 
+TOnlineAsyncOpHandle<FAuthGenerateAuthCode> FAuthCommon::GenerateAuthCode(FAuthGenerateAuthCode::Params&& Params)
+{
+	TOnlineAsyncOpRef<FAuthGenerateAuthCode> Operation = GetOp<FAuthGenerateAuthCode>(MoveTemp(Params));
+	Operation->SetError(Errors::NotImplemented());
+	return Operation->GetHandle();
+}
+
 TOnlineResult<FAuthGetAccountByPlatformUserId> FAuthCommon::GetAccountByPlatformUserId(FAuthGetAccountByPlatformUserId::Params&& Params)
 {
 	return TOnlineResult<FAuthGetAccountByPlatformUserId>(Errors::NotImplemented());
diff --git a/Engine/Plugins/Online/Experimental/OnlineServices/Source/OnlineServicesCommon/Private/Online/OnlineSyntaxTests.cpp b/Engine/Plugins/Online/Experimental/OnlineServices/Source/OnlineServicesCommon/Private/Online/OnlineSyntaxTests.cpp
--- a/Engine/Plugins/Online/Experimental/OnlineServices/Source/OnlineServicesCommon/Private/Online/OnlineSyntaxTests.cpp
+++ b/Engine/Plugins/Online/Experimental/OnlineServices/Source/OnlineServicesCommon/Private/Online/OnlineSyntaxTests.cpp
@@ -446,8 +446,12 @@ void TestToLogStringSyntax()
 	//ToLogString(UE::Online::FAuthLogin::Result()); default constructing a TSharedPtr will crash
 	ToLogString(UE::Online::FAuthLogout::Params());
 	ToLogString(UE::Online::FAuthLogout::Result());
-	ToLogString(UE::Online::FAuthGenerateAuth::Params());
-	ToLogString(UE::Online::FAuthGenerateAuth::Result());
+	ToLogString(UE::Online::FAuthGenerateAuthToken::Params());
+	ToLogString(UE::Online::FAuthGenerateAuthToken::Result());
+	ToLogString(UE::Online::FAuthGetAuthToken::Params());
+	ToLogString(UE::Online::FAuthGetAuthToken::Result());
+	ToLogString(UE::Online::FAuthGenerateAuthCode::Params());
+	ToLogString(UE::Online::FAuthGenerateAuthCode::Result());
 	ToLogString(UE::Online::FAuthGetAccountByPlatformUserId::Params());
 	ToLogString(UE::Online::FAuthGetAccountByPlatformUserId::Result());
 	ToLogString(UE::Online::FAuthGetAccountByAccountId::Params());
diff --git a/Engine/Plugins/Online/Experimental/OnlineServices/Source/OnlineServicesCommon/Public/Online/AuthCommon.h b/Engine/Plugins/Online/Experimental/OnlineServices/Source/OnlineServicesCommon/Public/Online/AuthCommon.h
--- a/Engine/Plugins/Online/Experimental/OnlineServices/Source/OnlineServicesCommon/Public/Online/AuthCommon.h
+++ b/Engine/Plugins/Online/Experimental/OnlineServices/Source/OnlineServicesCommon/Public/Online/AuthCommon.h
@@ -22,8 +22,9 @@ public:
 	// IAuth
 	virtual TOnlineAsyncOpHandle<FAuthLogin> Login(FAuthLogin::Params&& Params) override;
 	virtual TOnlineAsyncOpHandle<FAuthLogout> Logout(FAuthLogout::Params&& Params) override;
-	virtual TOnlineAsyncOpHandle<FAuthGenerateAuth> GenerateAuth(FAuthGenerateAuth::Params&& Params) override;
+	virtual TOnlineAsyncOpHandle<FAuthGenerateAuthToken> GenerateAuthToken(FAuthGenerateAuthToken::Params&& Params) override;
 	virtual TOnlineResult<FAuthGetAuthToken> GetAuthToken(FAuthGetAuthToken::Params&& Params) override;
+	virtual TOnlineAsyncOpHandle<FAuthGenerateAuthCode> GenerateAuthCode(FAuthGenerateAuthCode::Params&& Params) override;
 	virtual TOnlineResult<FAuthGetAccountByPlatformUserId> GetAccountByPlatformUserId(FAuthGetAccountByPlatformUserId::Params&& Params) override;
 	virtual TOnlineResult<FAuthGetAccountByAccountId> GetAccountByAccountId(FAuthGetAccountByAccountId::Params&& Params) override;
 	virtual TOnlineEvent<void(const FLoginStatusChanged&)> OnLoginStatusChanged() override;
diff --git a/Engine/Plugins/Online/Experimental/OnlineServices/Source/OnlineServicesInterface/Public/Online/Auth.h b/Engine/Plugins/Online/Experimental/OnlineServices/Source/OnlineServicesInterface/Public/Online/Auth.h
--- a/Engine/Plugins/Online/Experimental/OnlineServices/Source/OnlineServicesInterface/Public/Online/Auth.h
+++ b/Engine/Plugins/Online/Experimental/OnlineServices/Source/OnlineServicesInterface/Public/Online/Auth.h
@@ -75,9 +75,9 @@ struct FAuthLogout
 	};
 };
 
-struct FAuthGenerateAuth
+struct FAuthGenerateAuthToken
 {
-	static constexpr TCHAR Name[] = TEXT("GenerateAuth");
+	static constexpr TCHAR Name[] = TEXT("GenerateAuthToken");
 
 	struct Params
 	{
@@ -88,6 +88,8 @@ struct FAuthGenerateAuth
 
 	struct Result
 	{
+		FString Type;
+		FCredentialsToken Token;
 	};
 };
 
@@ -108,6 +110,23 @@ struct FAuthGetAuthToken
 	};
 };
 
+struct FAuthGenerateAuthCode
+{
+	static constexpr TCHAR Name[] = TEXT("GenerateAuthCode");
+
+	struct Params
+	{
+		FOnlineAccountIdHandle LocalUserId;
+		FString Type;
+		TArray<FString> Scopes;
+	};
+
+	struct Result
+	{
+		FString Code;
+	};
+};
+
 struct FAuthGetAccountByPlatformUserId
 {
 	static constexpr TCHAR Name[] = TEXT("GetAccountByPlatformUserId");
@@ -165,13 +184,18 @@ public:
 	/**
 	 *
 	 */
-	virtual TOnlineAsyncOpHandle<FAuthGenerateAuth> GenerateAuth(FAuthGenerateAuth::Params&& Params) = 0;
+	virtual TOnlineAsyncOpHandle<FAuthGenerateAuthToken> GenerateAuthToken(FAuthGenerateAuthToken::Params&& Params) = 0;
 
 	/**
 	 * Retrive a cached auth token
 	 */
 	virtual TOnlineResult<FAuthGetAuthToken> GetAuthToken(FAuthGetAuthToken::Params&& Params) = 0;
 
+	/**
+	 *
+	 */
+	virtual TOnlineAsyncOpHandle<FAuthGenerateAuthCode> GenerateAuthCode(FAuthGenerateAuthCode::Params&& Params) = 0;
+
 	/**
 	 * Get logged in user by local user num
 	 */ 
@@ -217,13 +241,15 @@ END_ONLINE_STRUCT_META()
 BEGIN_ONLINE_STRUCT_META(FAuthLogout::Result)
 END_ONLINE_STRUCT_META()
 
-BEGIN_ONLINE_STRUCT_META(FAuthGenerateAuth::Params)
-	ONLINE_STRUCT_FIELD(FAuthGenerateAuth::Params, LocalUserId),
-	ONLINE_STRUCT_FIELD(FAuthGenerateAuth::Params, Type),
-	ONLINE_STRUCT_FIELD(FAuthGenerateAuth::Params, Scopes)
+BEGIN_ONLINE_STRUCT_META(FAuthGenerateAuthToken::Params)
+	ONLINE_STRUCT_FIELD(FAuthGenerateAuthToken::Params, LocalUserId),
+	ONLINE_STRUCT_FIELD(FAuthGenerateAuthToken::Params, Type),
+	ONLINE_STRUCT_FIELD(FAuthGenerateAuthToken::Params, Scopes)
 END_ONLINE_STRUCT_META()
 
-BEGIN_ONLINE_STRUCT_META(FAuthGenerateAuth::Result)
+BEGIN_ONLINE_STRUCT_META(FAuthGenerateAuthToken::Result)
+	ONLINE_STRUCT_FIELD(FAuthGenerateAuthToken::Result, Type),
+	ONLINE_STRUCT_FIELD(FAuthGenerateAuthToken::Result, Token)
 END_ONLINE_STRUCT_META()
 
 BEGIN_ONLINE_STRUCT_META(FAuthGetAuthToken::Params)
@@ -236,6 +262,16 @@ BEGIN_ONLINE_STRUCT_META(FAuthGetAuthToken::Result)
 	ONLINE_STRUCT_FIELD(FAuthGetAuthToken::Result, Token)
 END_ONLINE_STRUCT_META()
 
+BEGIN_ONLINE_STRUCT_META(FAuthGenerateAuthCode::Params)
+ONLINE_STRUCT_FIELD(FAuthGenerateAuthCode::Params, LocalUserId),
+ONLINE_STRUCT_FIELD(FAuthGenerateAuthCode::Params, Type),
+ONLINE_STRUCT_FIELD(FAuthGenerateAuthCode::Params, Scopes)
+END_ONLINE_STRUCT_META()
+
+BEGIN_ONLINE_STRUCT_META(FAuthGenerateAuthCode::Result)
+	ONLINE_STRUCT_FIELD(FAuthGenerateAuthCode::Result, Code)
+END_ONLINE_STRUCT_META()
+
 BEGIN_ONLINE_STRUCT_META(FAuthGetAccountByPlatformUserId::Params)
 	ONLINE_STRUCT_FIELD(FAuthGetAccountByPlatformUserId::Params, PlatformUserId)
 END_ONLINE_STRUCT_META()
diff --git a/Engine/Plugins/Online/Experimental/OnlineServices/Source/OnlineServicesInterface/Public/Online/OnlineServices.h b/Engine/Plugins/Online/Experimental/OnlineServices/Source/OnlineServicesInterface/Public/Online/OnlineServices.h
--- a/Engine/Plugins/Online/Experimental/OnlineServices/Source/OnlineServicesInterface/Public/Online/OnlineServices.h
+++ b/Engine/Plugins/Online/Experimental/OnlineServices/Source/OnlineServicesInterface/Public/Online/OnlineServices.h
@@ -91,6 +91,12 @@ public:
 	 * Get the connectivity string used for client travel
 	 */
 	virtual TOnlineResult<FGetResolvedConnectString> GetResolvedConnectString(FGetResolvedConnectString::Params&& Params) = 0;
+
+	/**
+	 * Get the provider type of this instance.
+	 * @return provider type
+	 */
+	virtual EOnlineServices GetServicesProvider() const = 0;
 };
 
 /**
diff --git a/Engine/Plugins/Online/Experimental/OnlineServicesEOS/Source/Private/Online/AuthEOS.cpp b/Engine/Plugins/Online/Experimental/OnlineServicesEOS/Source/Private/Online/AuthEOS.cpp
--- a/Engine/Plugins/Online/Experimental/OnlineServicesEOS/Source/Private/Online/AuthEOS.cpp
+++ b/Engine/Plugins/Online/Experimental/OnlineServicesEOS/Source/Private/Online/AuthEOS.cpp
@@ -366,7 +366,7 @@ TOnlineAsyncOpHandle<FAuthLogin> FAuthEOS::Login(FAuthLogin::Params&& Params)
 	}
 	TOnlineAsyncOpRef<FAuthLogin> Op = GetOp<FAuthLogin>(MoveTemp(Params));
 	// Are we already logged in?
-	if (GetAccountIdByPlatformUserId(Params.PlatformUserId).IsOk())
+	if (GetAccountIdByPlatformUserId(Op->GetParams().PlatformUserId).IsOk())
 	{
 		Op->SetError(Errors::Auth::AlreadyLoggedIn());
 		return Op->GetHandle();
@@ -761,12 +761,6 @@ TOnlineAsyncOpHandle<FAuthLogout> FAuthEOS::Logout(FAuthLogout::Params&& Params)
 	return Op->GetHandle();
 }
 
-TOnlineAsyncOpHandle<FAuthGenerateAuth> FAuthEOS::GenerateAuth(FAuthGenerateAuth::Params&& Params)
-{
-	TSharedRef<TOnlineAsyncOp<FAuthGenerateAuth>> AsyncOperation = MakeShared<TOnlineAsyncOp<FAuthGenerateAuth>>(Services, MoveTemp(Params));
-	return AsyncOperation->GetHandle();
-}
-
 TOnlineResult<FAuthGetAccountByPlatformUserId> FAuthEOS::GetAccountByPlatformUserId(FAuthGetAccountByPlatformUserId::Params&& Params)
 {
 	TResult<FOnlineAccountIdHandle, FOnlineError> PlatformUserIdResult = GetAccountIdByPlatformUserId(Params.PlatformUserId);
diff --git a/Engine/Plugins/Online/Experimental/OnlineServicesEOS/Source/Private/Online/OnlineServicesEOS.cpp b/Engine/Plugins/Online/Experimental/OnlineServicesEOS/Source/Private/Online/OnlineServicesEOS.cpp
--- a/Engine/Plugins/Online/Experimental/OnlineServicesEOS/Source/Private/Online/OnlineServicesEOS.cpp
+++ b/Engine/Plugins/Online/Experimental/OnlineServicesEOS/Source/Private/Online/OnlineServicesEOS.cpp
@@ -68,14 +68,27 @@ void FOnlineServicesEOS::RegisterComponents()
 
 void FOnlineServicesEOS::Initialize()
 {
+	IEOSSDKManager* SDKManager = IEOSSDKManager::Get();
+	if (!SDKManager)
+	{
+		return;
+	}
+
+	EOS_EResult InitResult = SDKManager->Initialize();
+	if (InitResult != EOS_EResult::EOS_Success)
+	{
+		return;
+	}
+
 	FEOSPlatformConfig EOSPlatformConfig;
 	LoadConfig(EOSPlatformConfig);
 
-	FTCHARToUTF8 ProductId(*EOSPlatformConfig.ProductId);
-	FTCHARToUTF8 SandboxId(*EOSPlatformConfig.SandboxId);
-	FTCHARToUTF8 DeploymentId(*EOSPlatformConfig.DeploymentId);
-	FTCHARToUTF8 ClientId(*EOSPlatformConfig.ClientId);
-	FTCHARToUTF8 ClientSecret(*EOSPlatformConfig.ClientSecret);
+	const FTCHARToUTF8 ProductId(*EOSPlatformConfig.ProductId);
+	const FTCHARToUTF8 SandboxId(*EOSPlatformConfig.SandboxId);
+	const FTCHARToUTF8 DeploymentId(*EOSPlatformConfig.DeploymentId);
+	const FTCHARToUTF8 ClientId(*EOSPlatformConfig.ClientId);
+	const FTCHARToUTF8 ClientSecret(*EOSPlatformConfig.ClientSecret);
+	const FTCHARToUTF8 CacheDirectory(SDKManager->GetCacheDirBase() / TEXT("OnlineServicesEOS"));
 
 	EOS_Platform_Options PlatformOptions = {};
 	PlatformOptions.ApiVersion = EOS_PLATFORM_OPTIONS_API_LATEST;
@@ -85,28 +98,12 @@ void FOnlineServicesEOS::Initialize()
 	PlatformOptions.OverrideLocaleCode = nullptr;
 	PlatformOptions.Flags = EOS_PF_WINDOWS_ENABLE_OVERLAY_D3D9 | EOS_PF_WINDOWS_ENABLE_OVERLAY_D3D10 | EOS_PF_WINDOWS_ENABLE_OVERLAY_OPENGL; // Enable overlay support for D3D9/10 and OpenGL. This sample uses D3D11 or SDL.
 
-	char CacheDirectory[512];
-	FCStringAnsi::Strncpy(CacheDirectory, TCHAR_TO_UTF8(*(FPlatformProcess::UserDir() / FString(TEXT("CacheDir")))), 512);
-	PlatformOptions.CacheDirectory = CacheDirectory;
-
 	PlatformOptions.ProductId = ProductId.Get();
 	PlatformOptions.SandboxId = SandboxId.Get();
 	PlatformOptions.DeploymentId = DeploymentId.Get();
-
 	PlatformOptions.ClientCredentials.ClientId = ClientId.Get();
 	PlatformOptions.ClientCredentials.ClientSecret = ClientSecret.Get();
-
-	IEOSSDKManager* SDKManager = IEOSSDKManager::Get();
-	if (!SDKManager)
-	{
-		return;
-	}
-
-	EOS_EResult InitResult = SDKManager->Initialize();
-	if (InitResult != EOS_EResult::EOS_Success)
-	{
-		return;
-	}
+	PlatformOptions.CacheDirectory = CacheDirectory.Get();
 
 	EOSPlatformHandle = SDKManager->CreatePlatform(PlatformOptions);
 	if (EOSPlatformHandle)
diff --git a/Engine/Plugins/Online/Experimental/OnlineServicesEOS/Source/Public/Online/AuthEOS.h b/Engine/Plugins/Online/Experimental/OnlineServicesEOS/Source/Public/Online/AuthEOS.h
--- a/Engine/Plugins/Online/Experimental/OnlineServicesEOS/Source/Public/Online/AuthEOS.h
+++ b/Engine/Plugins/Online/Experimental/OnlineServicesEOS/Source/Public/Online/AuthEOS.h
@@ -25,7 +25,6 @@ public:
 	virtual void PreShutdown() override;
 	virtual TOnlineAsyncOpHandle<FAuthLogin> Login(FAuthLogin::Params&& Params) override;
 	virtual TOnlineAsyncOpHandle<FAuthLogout> Logout(FAuthLogout::Params&& Params) override;
-	virtual TOnlineAsyncOpHandle<FAuthGenerateAuth> GenerateAuth(FAuthGenerateAuth::Params&& Params) override;
 	virtual TOnlineResult<FAuthGetAccountByPlatformUserId> GetAccountByPlatformUserId(FAuthGetAccountByPlatformUserId::Params&& Params) override;
 	virtual TOnlineResult<FAuthGetAccountByAccountId> GetAccountByAccountId(FAuthGetAccountByAccountId::Params&& Params) override;
 
diff --git a/Engine/Plugins/Online/Experimental/OnlineServicesEOS/Source/Public/Online/OnlineServicesEOS.h b/Engine/Plugins/Online/Experimental/OnlineServicesEOS/Source/Public/Online/OnlineServicesEOS.h
--- a/Engine/Plugins/Online/Experimental/OnlineServicesEOS/Source/Public/Online/OnlineServicesEOS.h
+++ b/Engine/Plugins/Online/Experimental/OnlineServicesEOS/Source/Public/Online/OnlineServicesEOS.h
@@ -26,8 +26,7 @@ public:
 	virtual void RegisterComponents() override;
 	virtual void Initialize() override;
 	virtual TOnlineResult<FGetResolvedConnectString> GetResolvedConnectString(FGetResolvedConnectString::Params&& Params) override;
-
-	static EOnlineServices GetServicesProvider() { return EOnlineServices::Epic; }
+	virtual EOnlineServices GetServicesProvider() const override { return EOnlineServices::Epic; }
 
 	EOS_HPlatform GetEOSPlatformHandle() const;
 protected:
diff --git a/Engine/Plugins/Online/Experimental/OnlineServicesNull/Source/Private/Online/AuthNull.cpp b/Engine/Plugins/Online/Experimental/OnlineServicesNull/Source/Private/Online/AuthNull.cpp
--- a/Engine/Plugins/Online/Experimental/OnlineServicesNull/Source/Private/Online/AuthNull.cpp
+++ b/Engine/Plugins/Online/Experimental/OnlineServicesNull/Source/Private/Online/AuthNull.cpp
@@ -114,12 +114,6 @@ TOnlineAsyncOpHandle<FAuthLogout> FAuthNull::Logout(FAuthLogout::Params&& Params
 	return Op->GetHandle();
 }
 
-TOnlineAsyncOpHandle<FAuthGenerateAuth> FAuthNull::GenerateAuth(FAuthGenerateAuth::Params&& Params)
-{
-	TSharedRef<TOnlineAsyncOp<FAuthGenerateAuth>> AsyncOperation = MakeShared<TOnlineAsyncOp<FAuthGenerateAuth>>(Services, MoveTemp(Params));
-	return AsyncOperation->GetHandle();
-}
-
 TOnlineResult<FAuthGetAccountByAccountId> FAuthNull::GetAccountByAccountId(FAuthGetAccountByAccountId::Params&& Params)
 {
 	if (TSharedRef<FAccountInfoNull>* const FoundAccount = AccountInfos.Find(Params.LocalUserId))
diff --git a/Engine/Plugins/Online/Experimental/OnlineServicesNull/Source/Public/Online/AuthNull.h b/Engine/Plugins/Online/Experimental/OnlineServicesNull/Source/Public/Online/AuthNull.h
--- a/Engine/Plugins/Online/Experimental/OnlineServicesNull/Source/Public/Online/AuthNull.h
+++ b/Engine/Plugins/Online/Experimental/OnlineServicesNull/Source/Public/Online/AuthNull.h
@@ -56,7 +56,6 @@ public:
 	virtual void PreShutdown() override;
 	virtual TOnlineAsyncOpHandle<FAuthLogin> Login(FAuthLogin::Params&& Params) override;
 	virtual TOnlineAsyncOpHandle<FAuthLogout> Logout(FAuthLogout::Params&& Params) override;
-	virtual TOnlineAsyncOpHandle<FAuthGenerateAuth> GenerateAuth(FAuthGenerateAuth::Params&& Params) override;
 	virtual TOnlineResult<FAuthGetAccountByPlatformUserId> GetAccountByPlatformUserId(FAuthGetAccountByPlatformUserId::Params&& Params) override;
 	virtual TOnlineResult<FAuthGetAccountByAccountId> GetAccountByAccountId(FAuthGetAccountByAccountId::Params&& Params) override;
 
diff --git a/Engine/Plugins/Online/Experimental/OnlineServicesNull/Source/Public/Online/OnlineServicesNull.h b/Engine/Plugins/Online/Experimental/OnlineServicesNull/Source/Public/Online/OnlineServicesNull.h
--- a/Engine/Plugins/Online/Experimental/OnlineServicesNull/Source/Public/Online/OnlineServicesNull.h
+++ b/Engine/Plugins/Online/Experimental/OnlineServicesNull/Source/Public/Online/OnlineServicesNull.h
@@ -15,8 +15,8 @@ public:
 	virtual void RegisterComponents() override;
 	virtual void Initialize() override;
 	virtual TOnlineResult<FGetResolvedConnectString> GetResolvedConnectString(FGetResolvedConnectString::Params&& Params) override;
+	virtual EOnlineServices GetServicesProvider() const override { return EOnlineServices::Null; }
 
-	static EOnlineServices GetServicesProvider() { return EOnlineServices::Null; }
 };
 
 /* UE::Online */ }
diff --git a/Engine/Plugins/Online/Experimental/OnlineServicesOSSAdapter/Source/Private/Online/AuthOSSAdapter.cpp b/Engine/Plugins/Online/Experimental/OnlineServicesOSSAdapter/Source/Private/Online/AuthOSSAdapter.cpp
--- a/Engine/Plugins/Online/Experimental/OnlineServicesOSSAdapter/Source/Private/Online/AuthOSSAdapter.cpp
+++ b/Engine/Plugins/Online/Experimental/OnlineServicesOSSAdapter/Source/Private/Online/AuthOSSAdapter.cpp
@@ -209,78 +209,153 @@ TOnlineAsyncOpHandle<FAuthLogout> FAuthOSSAdapter::Logout(FAuthLogout::Params&&
 	return Op->GetHandle();
 }
 
-TOnlineAsyncOpHandle<FAuthGenerateAuth> FAuthOSSAdapter::GenerateAuth(FAuthGenerateAuth::Params&& Params)
+TOnlineAsyncOpHandle<FAuthGenerateAuthToken> FAuthOSSAdapter::GenerateAuthToken(FAuthGenerateAuthToken::Params&& Params)
 {
-	TSharedRef<TOnlineAsyncOp<FAuthGenerateAuth>> Op = GetOp<FAuthGenerateAuth>(MoveTemp(Params));
+	TSharedRef<TOnlineAsyncOp<FAuthGenerateAuthToken>> Op = GetOp<FAuthGenerateAuthToken>(MoveTemp(Params));
 
 	if (!Op->IsReady())
 	{
-		FUniqueNetIdRef UniqueNetId = GetUniqueNetId(Params.LocalUserId);
-
+		const FUniqueNetIdRef UniqueNetId = GetUniqueNetId(Op->GetParams().LocalUserId);
 		if (!UniqueNetId->IsValid())
 		{
 			Op->SetError(Errors::InvalidParams());
 			return Op->GetHandle();
 		}
 
-		if (GetIdentityInterface()->GetLoginStatus(*UniqueNetId) != ::ELoginStatus::LoggedIn)
+		IOnlineIdentityPtr Identity = GetIdentityInterface();
+		if (Identity->GetLoginStatus(*UniqueNetId) == ::ELoginStatus::LoggedIn)
 		{
-			Op->SetError(Errors::InvalidAuth());
+			Op->Then([this](TOnlineAsyncOp<FAuthGenerateAuthToken>& Op)
+			{
+				if (IOnlineIdentityPtr Identity = GetIdentityInterface())
+				{
+					const FUniqueNetIdRef UniqueNetId = GetUniqueNetId(Op.GetParams().LocalUserId);
+					const int32 LocalUserNum = Identity->GetLocalUserNumFromPlatformUserId(Identity->GetPlatformUserIdFromUniqueNetId(*UniqueNetId));
+
+					Identity->GetLinkedAccountAuthToken(LocalUserNum, IOnlineIdentity::FOnGetLinkedAccountAuthTokenCompleteDelegate::CreateLambda([this, WeakOp = Op.AsWeak()](int32 LocalUserNum, bool bWasSuccessful, const FExternalAuthToken& AuthToken)
+					{
+						TSharedPtr<TOnlineAsyncOp<FAuthGenerateAuthToken>> Op = WeakOp.Pin();
+						if (!Op)
+						{
+							return;
+						}
+
+						if (bWasSuccessful && AuthToken.IsValid())
+						{
+							FAuthGenerateAuthToken::Result Result;
+							Result.Type = LexToString(Services.GetServicesProvider());
+							if (AuthToken.HasTokenString())
+							{
+								Result.Token.Emplace<FString>(AuthToken.TokenString);
+							}
+							else if (AuthToken.HasTokenData())
+							{
+								Result.Token.Emplace<TArray<uint8>>(AuthToken.TokenData);
+							}
+							else
+							{
+								checkNoEntry();
+							}
+							Op->SetResult(MoveTemp(Result));
+						}
+						else
+						{
+							Op->SetError(Errors::Unknown());
+						}
+					}));
+				}
+				else
+				{
+					Op.SetError(Errors::Unknown());
+				}
+			})
+			.Enqueue(GetSerialQueue());
 		}
 		else
 		{
-			Op->SetResult(FAuthGenerateAuth::Result());
+			Op->SetError(Errors::InvalidAuth());
 		}
 	}
 
 	return Op->GetHandle();
 }
 
-TOnlineResult<FAuthGetAuthToken> FAuthOSSAdapter::GetAuthToken(FAuthGetAuthToken::Params&& Params)
+TOnlineAsyncOpHandle<FAuthGenerateAuthCode> FAuthOSSAdapter::GenerateAuthCode(FAuthGenerateAuthCode::Params&& Params)
 {
-	FUniqueNetIdRef UniqueNetId = GetUniqueNetId(Params.LocalUserId);
+	TSharedRef<TOnlineAsyncOp<FAuthGenerateAuthCode>> Op = GetOp<FAuthGenerateAuthCode>(MoveTemp(Params));
 
-	if (!UniqueNetId->IsValid())
+	if (!Op->IsReady())
 	{
-		return TOnlineResult<FAuthGetAuthToken>(Errors::InvalidUser());
-	}
+		const FUniqueNetIdRef UniqueNetId = GetUniqueNetId(Op->GetParams().LocalUserId);
+		if (!UniqueNetId->IsValid())
+		{
+			Op->SetError(Errors::InvalidParams());
+			return Op->GetHandle();
+		}
 
-	if (GetIdentityInterface()->GetLoginStatus(*UniqueNetId) != ::ELoginStatus::LoggedIn)
-	{
-		return TOnlineResult<FAuthGetAuthToken>(Errors::InvalidAuth());
-	}
+		IOnlineIdentityPtr Identity = GetIdentityInterface();
+		if (Identity->GetLoginStatus(*UniqueNetId) == ::ELoginStatus::LoggedIn)
+		{
+			const int32 LocalUserNum = Identity->GetLocalUserNumFromPlatformUserId(Identity->GetPlatformUserIdFromUniqueNetId(*UniqueNetId));
 
-	TSharedPtr<FUserOnlineAccount> UserAccount = GetIdentityInterface()->GetUserAccount(*UniqueNetId);
-	if (UserAccount.IsValid())
-	{
-		FAuthGetAuthToken::Result Result;
-		Result.Token = UserAccount->GetAccessToken();
+			FString Token = Identity->GetAuthToken(LocalUserNum);
+			if (Token.IsEmpty())
+			{
+				Op->Then([this](TOnlineAsyncOp<FAuthGenerateAuthCode>& Op)
+				{
+					// AutoLogin to generate the token. Token generation is a side effect of calling login, but you can be logged in without calling login.
+					FOnLoginCompleteDelegate LoginCompleteDelegate = FOnLoginCompleteDelegate::CreateLambda([this, WeakOp = Op.AsWeak()](int32 LocalUserNum, bool bWasSuccessful, const FUniqueNetId& UserId, const FString& Error)
+					{
+						TSharedPtr<TOnlineAsyncOp<FAuthGenerateAuthCode>> Op = WeakOp.Pin();
+						if (!Op)
+						{
+							return;
+						}
 
-		return TOnlineResult<FAuthGetAuthToken>(Result);
-	}
-	else
-	{
-		// Fall back to GetAuthToken if the user account is not available (Steam doesn't implement GetUserAccount)
-		int32 LocalUserIndex = GetLocalUserNum(Params.LocalUserId);
+						IOnlineIdentityPtr Identity = GetIdentityInterface();
+						if (bWasSuccessful && Identity)
+						{
+							if (const FDelegateHandle* LoginCompleteHandle = Op->Data.Get<FDelegateHandle>(TEXT("LoginCompleteHandle")))
+							{
+								Identity->ClearOnLoginCompleteDelegate_Handle(LocalUserNum, *const_cast<FDelegateHandle*>(LoginCompleteHandle));
+							}
 
-		if (LocalUserIndex < 0 || LocalUserIndex >= MAX_LOCAL_PLAYERS)
-		{
-			return TOnlineResult<FAuthGetAuthToken>(Errors::InvalidUser());
-		}
+							FString Token = Identity->GetAuthToken(LocalUserNum);
+							if (!Token.IsEmpty())
+							{
+								FAuthGenerateAuthCode::Result Result;
+								Result.Code = MoveTemp(Token);
+								Op->SetResult(MoveTemp(Result));
+								return;
+							}
+						}
+						Op->SetError(Errors::Unknown());
+					});
 
-		FString Token = GetIdentityInterface()->GetAuthToken(LocalUserIndex);
-		if (!Token.IsEmpty())
-		{
-			FAuthGetAuthToken::Result Result;
-			Result.Token = Token;
 
-			return TOnlineResult<FAuthGetAuthToken>(Result);
+					IOnlineIdentityPtr Identity = GetIdentityInterface();
+					const FUniqueNetIdRef UniqueNetId = GetUniqueNetId(Op.GetParams().LocalUserId);
+					const int32 LocalUserNum = Identity->GetLocalUserNumFromPlatformUserId(Identity->GetPlatformUserIdFromUniqueNetId(*UniqueNetId));
+					FDelegateHandle LoginCompleteHandle = Identity->AddOnLoginCompleteDelegate_Handle(LocalUserNum, LoginCompleteDelegate);
+					Op.Data.Set(TEXT("LoginCompleteHandle"), MoveTemp(LoginCompleteHandle));
+					Identity->AutoLogin(LocalUserNum);
+				})
+				.Enqueue(GetSerialQueue());
+			}
+			else
+			{
+				FAuthGenerateAuthCode::Result Result;
+				Result.Code = MoveTemp(Token);
+				Op->SetResult(MoveTemp(Result));
+			}
 		}
 		else
 		{
-			return TOnlineResult<FAuthGetAuthToken>(Errors::InvalidUser());
+			Op->SetError(Errors::InvalidAuth());
 		}
 	}
+
+	return Op->GetHandle();
 }
 
 TOnlineResult<FAuthGetAccountByPlatformUserId> FAuthOSSAdapter::GetAccountByPlatformUserId(FAuthGetAccountByPlatformUserId::Params&& Params)
diff --git a/Engine/Plugins/Online/Experimental/OnlineServicesOSSAdapter/Source/Public/Online/AuthOSSAdapter.h b/Engine/Plugins/Online/Experimental/OnlineServicesOSSAdapter/Source/Public/Online/AuthOSSAdapter.h
--- a/Engine/Plugins/Online/Experimental/OnlineServicesOSSAdapter/Source/Public/Online/AuthOSSAdapter.h
+++ b/Engine/Plugins/Online/Experimental/OnlineServicesOSSAdapter/Source/Public/Online/AuthOSSAdapter.h
@@ -25,8 +25,8 @@ public:
 	// IAuth
 	virtual TOnlineAsyncOpHandle<FAuthLogin> Login(FAuthLogin::Params&& Params) override;
 	virtual TOnlineAsyncOpHandle<FAuthLogout> Logout(FAuthLogout::Params&& Params) override;
-	virtual TOnlineAsyncOpHandle<FAuthGenerateAuth> GenerateAuth(FAuthGenerateAuth::Params&& Params) override;
-	virtual TOnlineResult<FAuthGetAuthToken> GetAuthToken(FAuthGetAuthToken::Params&& Params) override;
+	virtual TOnlineAsyncOpHandle<FAuthGenerateAuthToken> GenerateAuthToken(FAuthGenerateAuthToken::Params&& Params) override;
+	virtual TOnlineAsyncOpHandle<FAuthGenerateAuthCode> GenerateAuthCode(FAuthGenerateAuthCode::Params&& Params) override;
 	virtual TOnlineResult<FAuthGetAccountByPlatformUserId> GetAccountByPlatformUserId(FAuthGetAccountByPlatformUserId::Params&& Params) override;
 	virtual TOnlineResult<FAuthGetAccountByAccountId> GetAccountByAccountId(FAuthGetAccountByAccountId::Params&& Params) override;
 
diff --git a/Engine/Plugins/Online/Experimental/OnlineServicesOSSAdapter/Source/Public/Online/OnlineServicesOSSAdapter.h b/Engine/Plugins/Online/Experimental/OnlineServicesOSSAdapter/Source/Public/Online/OnlineServicesOSSAdapter.h
--- a/Engine/Plugins/Online/Experimental/OnlineServicesOSSAdapter/Source/Public/Online/OnlineServicesOSSAdapter.h
+++ b/Engine/Plugins/Online/Experimental/OnlineServicesOSSAdapter/Source/Public/Online/OnlineServicesOSSAdapter.h
@@ -19,6 +19,7 @@ public:
 
 	ONLINESERVICESOSSADAPTER_API virtual void RegisterComponents() override;
 	ONLINESERVICESOSSADAPTER_API virtual void Initialize() override;
+	virtual EOnlineServices GetServicesProvider() const override { return ServicesType; }
 
 	IOnlineSubsystem& GetSubsystem() { return *Subsystem; }
 	FOnlineUniqueNetIdRegistry& GetAccountIdRegistry() { return *AccountIdRegistry; }
diff --git a/Engine/Plugins/Online/OnlineSubsystemEOS/Source/OnlineSubsystemEOS/Private/EOSHelpers.cpp b/Engine/Plugins/Online/OnlineSubsystemEOS/Source/OnlineSubsystemEOS/Private/EOSHelpers.cpp
--- a/Engine/Plugins/Online/OnlineSubsystemEOS/Source/OnlineSubsystemEOS/Private/EOSHelpers.cpp
+++ b/Engine/Plugins/Online/OnlineSubsystemEOS/Source/OnlineSubsystemEOS/Private/EOSHelpers.cpp
@@ -9,12 +9,6 @@
 #include "OnlineSubsystemEOS.h"
 #include "UserManagerEOS.h"
 
-FString FEOSHelpers::PlatformCreateCacheDir(const FString &ArtifactName, const FString &EOSSettingsCacheDir)
-{
-	FString CacheDir = FPlatformProcess::UserDir() / ArtifactName / EOSSettingsCacheDir;
-	return MoveTemp(CacheDir);
-}
-
 void FEOSHelpers::PlatformAuthCredentials(EOS_Auth_Credentials &Credentials)
 {
 }
diff --git a/Engine/Plugins/Online/OnlineSubsystemEOS/Source/OnlineSubsystemEOS/Private/EOSHelpers.h b/Engine/Plugins/Online/OnlineSubsystemEOS/Source/OnlineSubsystemEOS/Private/EOSHelpers.h
--- a/Engine/Plugins/Online/OnlineSubsystemEOS/Source/OnlineSubsystemEOS/Private/EOSHelpers.h
+++ b/Engine/Plugins/Online/OnlineSubsystemEOS/Source/OnlineSubsystemEOS/Private/EOSHelpers.h
@@ -20,7 +20,6 @@ class FEOSHelpers
 public:
 	virtual ~FEOSHelpers() = default;
 
-	virtual FString PlatformCreateCacheDir(const FString &ArtifactName, const FString &EOSSettingsCacheDir);
 	virtual void PlatformAuthCredentials(EOS_Auth_Credentials &Credentials);
 	virtual void PlatformTriggerLoginUI(FOnlineSubsystemEOS* EOSSubsystem, const int ControllerIndex, bool bShowOnlineOnly, bool bShowSkipButton, const FOnLoginUIClosedDelegate& Delegate);
 	virtual IEOSPlatformHandlePtr CreatePlatform(EOS_Platform_Options& PlatformOptions);
diff --git a/Engine/Plugins/Online/OnlineSubsystemEOS/Source/OnlineSubsystemEOS/Private/IOS/IOSEOSHelpers.cpp b/Engine/Plugins/Online/OnlineSubsystemEOS/Source/OnlineSubsystemEOS/Private/IOS/IOSEOSHelpers.cpp
--- a/Engine/Plugins/Online/OnlineSubsystemEOS/Source/OnlineSubsystemEOS/Private/IOS/IOSEOSHelpers.cpp
+++ b/Engine/Plugins/Online/OnlineSubsystemEOS/Source/OnlineSubsystemEOS/Private/IOS/IOSEOSHelpers.cpp
@@ -53,24 +53,6 @@ FIOSEOSHelpers::~FIOSEOSHelpers()
 #endif
 }
 
-FString FIOSEOSHelpers::PlatformCreateCacheDir(const FString &ArtifactName, const FString &EOSSettingsCacheDir)
-{
-	NSString *BundleIdentifier = [[NSBundle mainBundle] bundleIdentifier];
-	NSString *CacheDirectory = NSTemporaryDirectory();									// Potentially use NSCachesDirectory
-	CacheDirectory = [CacheDirectory stringByAppendingPathComponent:BundleIdentifier];
-
-	FTCHARToUTF8 TCArtifactName(*ArtifactName);
-	NSString *NSArtifactName = [NSString stringWithUTF8String:TCArtifactName.Get()];
-	CacheDirectory = [CacheDirectory stringByAppendingPathComponent:NSArtifactName];
-
-	FTCHARToUTF8 TCEOSSettingsCacheDir(*EOSSettingsCacheDir);
-	NSString *NSEOSSettingsCacheDir = [NSString stringWithUTF8String:TCArtifactName.Get()];
-	CacheDirectory = [CacheDirectory stringByAppendingPathComponent:NSEOSSettingsCacheDir];
-
-	const char *CStrCacheDirectory = [CacheDirectory UTF8String];
-	return FString(UTF8_TO_TCHAR(CStrCacheDirectory));
-}
-
 void FIOSEOSHelpers::PlatformAuthCredentials(EOS_Auth_Credentials &Credentials)
 {
 #if PRESENTATIONCONTEXT_REQUIRED
diff --git a/Engine/Plugins/Online/OnlineSubsystemEOS/Source/OnlineSubsystemEOS/Private/IOS/IOSEOSHelpers.h b/Engine/Plugins/Online/OnlineSubsystemEOS/Source/OnlineSubsystemEOS/Private/IOS/IOSEOSHelpers.h
--- a/Engine/Plugins/Online/OnlineSubsystemEOS/Source/OnlineSubsystemEOS/Private/IOS/IOSEOSHelpers.h
+++ b/Engine/Plugins/Online/OnlineSubsystemEOS/Source/OnlineSubsystemEOS/Private/IOS/IOSEOSHelpers.h
@@ -16,7 +16,6 @@ class FIOSEOSHelpers : FEOSHelpers
 public:
 	virtual ~FIOSEOSHelpers();
 
-	virtual FString PlatformCreateCacheDir(const FString &ArtifactName, const FString &EOSSettingsCacheDir) override;
 	virtual void PlatformAuthCredentials(EOS_Auth_Credentials &Credentials) override;
 	virtual void PlatformTriggerLoginUI(FOnlineSubsystemEOS* EOSSubsystem, const int ControllerIndex, bool bShowOnlineOnly, bool bShowSkipButton, const FOnLoginUIClosedDelegate& Delegate) override;
 
diff --git a/Engine/Plugins/Online/OnlineSubsystemEOS/Source/OnlineSubsystemEOS/Private/OnlineSubsystemEOS.cpp b/Engine/Plugins/Online/OnlineSubsystemEOS/Source/OnlineSubsystemEOS/Private/OnlineSubsystemEOS.cpp
--- a/Engine/Plugins/Online/OnlineSubsystemEOS/Source/OnlineSubsystemEOS/Private/OnlineSubsystemEOS.cpp
+++ b/Engine/Plugins/Online/OnlineSubsystemEOS/Source/OnlineSubsystemEOS/Private/OnlineSubsystemEOS.cpp
@@ -226,7 +226,7 @@ bool FOnlineSubsystemEOS::PlatformCreate()
 	PlatformOptions.Flags = IsRunningGame() ? OverlayFlags : EOS_PF_DISABLE_OVERLAY;
 	// Make the cache directory be in the user's writable area
 
-	const FString CacheDir = EOSHelpersPtr->PlatformCreateCacheDir(ArtifactName, EOSSettings.CacheDir);
+	const FString CacheDir = EOSSDKManager->GetCacheDirBase() / ArtifactName / EOSSettings.CacheDir;
 	FCStringAnsi::Strncpy(PlatformOptions.CacheDirectoryAnsi, TCHAR_TO_UTF8(*CacheDir), EOS_OSS_STRING_BUFFER_LENGTH);
 	FCStringAnsi::Strncpy(PlatformOptions.EncryptionKeyAnsi, TCHAR_TO_UTF8(*ArtifactSettings.EncryptionKey), EOS_ENCRYPTION_KEY_MAX_BUFFER_LEN);
 
