V1 Delegate Adapter- Adds the MakeDelegateAdapter and MakeDelegateAdapter functions that allow us to bind generic unique functions to delegate to improve cross compatibility between OSSV1 and V2. Also includes some tests and example implementations in the auth adapter

diff --git a/Engine/Plugins/Online/Experimental/OnlineServicesOSSAdapter/Source/Private/Online/AuthOSSAdapter.cpp b/Engine/Plugins/Online/Experimental/OnlineServicesOSSAdapter/Source/Private/Online/AuthOSSAdapter.cpp
--- a/Engine/Plugins/Online/Experimental/OnlineServicesOSSAdapter/Source/Private/Online/AuthOSSAdapter.cpp
+++ b/Engine/Plugins/Online/Experimental/OnlineServicesOSSAdapter/Source/Private/Online/AuthOSSAdapter.cpp
@@ -5,6 +5,7 @@
 #include "Online/OnlineServicesOSSAdapter.h"
 #include "Online/OnlineIdOSSAdapter.h"
 #include "Online/OnlineErrorDefinitions.h"
+#include "Online/DelegateAdapter.h"
 
 #include "OnlineSubsystem.h"
 #include "Interfaces/OnlineIdentityInterface.h"
@@ -96,37 +97,25 @@ TOnlineAsyncOpHandle<FAuthLogin> FAuthOSSAdapter::Login(FAuthLogin::Params&& Par
 				Credentials.Id = Op.GetParams().CredentialsId;
 				Credentials.Token = Op.GetParams().CredentialsToken.Get<FString>(); // TODO: handle binary token
 
-				TSharedPtr<TPair<FDelegateHandle, TPromise<TOnlineResult<FAuthLogin>>>> ResultPtr = MakeShared<TPair<FDelegateHandle, TPromise<TOnlineResult<FAuthLogin>>>>(FDelegateHandle(), MoveTemp(Result));
-				ResultPtr->Get<FDelegateHandle>() = GetIdentityInterface()->OnLoginCompleteDelegates[LocalUserIndex].AddLambda(
-					[WeakThis = TWeakPtr<IAuth>(AsShared()), this, ResultPtr, PlatformUserId = Op.GetParams().PlatformUserId, LocalUserIndex](int32 LocalUserNum, bool bWasSuccessful, const FUniqueNetId& UserId, const FString& Error) mutable
+				MakeMulticastAdapter(this, GetIdentityInterface()->OnLoginCompleteDelegates[LocalUserIndex], 
+				[this, ResultPromise = MoveTemp(Result), PlatformUserId = Op.GetParams().PlatformUserId, LocalUserIndex](int32 LocalUserNum, bool bWasSuccessful, const FUniqueNetId& UserId, const FString& Error) mutable
+				{
+					if (bWasSuccessful)
 					{
-						if (LocalUserIndex != LocalUserNum)
-						{
-							return;
-						}
-
-						TSharedPtr<IAuth> PinnedThis = WeakThis.Pin();
-						if (PinnedThis.IsValid())
-						{
-							TPromise<TOnlineResult<FAuthLogin>>& OnlineResult = ResultPtr->Get<TPromise<TOnlineResult<FAuthLogin>>>();
-							if (bWasSuccessful)
-							{
-								FOnlineAccountIdHandle Handle = static_cast<FOnlineServicesOSSAdapter&>(Services).GetAccountIdRegistry().FindOrAddHandle(UserId.AsShared());
-								FAuthLogin::Result Result = { MakeShared<FAccountInfo>() };
-								Result.AccountInfo->PlatformUserId = PlatformUserId;
-								Result.AccountInfo->UserId = Handle;
-								Result.AccountInfo->LoginStatus = ELoginStatus::LoggedIn;
-								OnlineResult.EmplaceValue(MoveTemp(Result));
-							}
-							else
-							{
-								FOnlineError V2Error = Errors::Unknown(); // TODO: V1 to V2 error conversion/error from string conversion
-								OnlineResult.EmplaceValue(MoveTemp(V2Error));
-							}
-						}
+						FOnlineAccountIdHandle Handle = static_cast<FOnlineServicesOSSAdapter&>(Services).GetAccountIdRegistry().FindOrAddHandle(UserId.AsShared());
+						FAuthLogin::Result Result = { MakeShared<FAccountInfo>() };
+						Result.AccountInfo->PlatformUserId = PlatformUserId;
+						Result.AccountInfo->UserId = Handle;
+						Result.AccountInfo->LoginStatus = ELoginStatus::LoggedIn;
+						ResultPromise.EmplaceValue(MoveTemp(Result));
+					}
+					else
+					{
+						FOnlineError V2Error = Errors::Unknown(); // TODO: V1 to V2 error conversion/error from string conversion
+						ResultPromise.EmplaceValue(MoveTemp(V2Error));
+					}
 
-						GetIdentityInterface()->OnLoginCompleteDelegates[LocalUserIndex].Remove(ResultPtr->Get<FDelegateHandle>());
-					});
+				});
 
 				GetIdentityInterface()->Login(LocalUserIndex, Credentials);
 			})
@@ -162,33 +151,19 @@ TOnlineAsyncOpHandle<FAuthLogout> FAuthOSSAdapter::Logout(FAuthLogout::Params&&
 
 		Op->Then([this, LocalUserIndex](TOnlineAsyncOp<FAuthLogout>& Op, TPromise<TOnlineResult<FAuthLogout>>&& Result)
 			{
-				TSharedPtr<TPair<FDelegateHandle, TPromise<TOnlineResult<FAuthLogout>>>> ResultPtr = MakeShared<TPair<FDelegateHandle, TPromise<TOnlineResult<FAuthLogout>>>>(FDelegateHandle(), MoveTemp(Result));
-				ResultPtr->Get<FDelegateHandle>() = GetIdentityInterface()->OnLogoutCompleteDelegates[LocalUserIndex].AddLambda(
-					[WeakThis = TWeakPtr<IAuth>(AsShared()), this, ResultPtr, LocalUserIndex](int32 LocalUserNum, bool bWasSuccessful) mutable
+				MakeMulticastAdapter(this, GetIdentityInterface()->OnLogoutCompleteDelegates[LocalUserIndex],
+				[this, OnlineResult = MoveTemp(Result), LocalUserIndex](int32 LocalUserNum, bool bWasSuccessful) mutable
 				{
-					if (LocalUserIndex != LocalUserNum)
+					if (bWasSuccessful)
 					{
-						return;
+						FAuthLogout::Result Result;
+						OnlineResult.EmplaceValue(MoveTemp(Result));
 					}
-
-					TPromise<TOnlineResult<FAuthLogout>>& OnlineResult = ResultPtr->Get<TPromise<TOnlineResult<FAuthLogout>>>();
-
-					TSharedPtr<IAuth> PinnedThis = WeakThis.Pin();
-					if (PinnedThis.IsValid())
+					else
 					{
-						if (bWasSuccessful)
-						{
-							FAuthLogout::Result Result;
-							OnlineResult.EmplaceValue(MoveTemp(Result));
-						}
-						else
-						{
-							FOnlineError V2Error = Errors::Unknown(); // TODO: V1 to V2 error conversion/error from string conversion
-							OnlineResult.EmplaceValue(MoveTemp(V2Error));
-						}
+						FOnlineError V2Error = Errors::Unknown(); // TODO: V1 to V2 error conversion/error from string conversion
+						OnlineResult.EmplaceValue(MoveTemp(V2Error));
 					}
-
-					GetIdentityInterface()->OnLogoutCompleteDelegates[LocalUserIndex].Remove(ResultPtr->Get<FDelegateHandle>());
 				});
 
 				GetIdentityInterface()->Logout(LocalUserIndex);
@@ -232,7 +207,8 @@ TOnlineAsyncOpHandle<FAuthGenerateAuthToken> FAuthOSSAdapter::GenerateAuthToken(
 					const FUniqueNetIdRef UniqueNetId = GetUniqueNetId(Op.GetParams().LocalUserId);
 					const int32 LocalUserNum = Identity->GetLocalUserNumFromPlatformUserId(Identity->GetPlatformUserIdFromUniqueNetId(*UniqueNetId));
 
-					Identity->GetLinkedAccountAuthToken(LocalUserNum, IOnlineIdentity::FOnGetLinkedAccountAuthTokenCompleteDelegate::CreateLambda([this, WeakOp = Op.AsWeak()](int32 LocalUserNum, bool bWasSuccessful, const FExternalAuthToken& AuthToken)
+					Identity->GetLinkedAccountAuthToken(LocalUserNum, *MakeDelegateAdapter(this, 
+					[this, WeakOp = Op.AsWeak()](int32 LocalUserNum, bool bWasSuccessful, const FExternalAuthToken& AuthToken)
 					{
 						TSharedPtr<TOnlineAsyncOp<FAuthGenerateAuthToken>> Op = WeakOp.Pin();
 						if (!Op)
diff --git a/Engine/Plugins/Online/Experimental/OnlineServicesOSSAdapter/Source/Public/Online/DelegateAdapter.h b/Engine/Plugins/Online/Experimental/OnlineServicesOSSAdapter/Source/Public/Online/DelegateAdapter.h
new file mode 100644
--- /dev/null
+++ b/Engine/Plugins/Online/Experimental/OnlineServicesOSSAdapter/Source/Public/Online/DelegateAdapter.h
@@ -0,0 +1,143 @@
+// Copyright Epic Games, Inc. All Rights Reserved.
+#pragma once
+#include "OnlineSubsystemTypes.h"
+
+// this isn't actually used, but it allows files to only include this instead of both for ease-of-use
+#include "Online/MulticastAdapter.h"
+
+/*
+	BASIC USAGE
+		The Delegate Adapter can be used in place of an FOn{Name}Delegate::CreateLambda call
+
+		i.e.
+			
+			Identity->GetAuthToken(LocalUserNum, FOnAuthTokenGetComplete::CreateLambda[this](int32 UserNum, FString LoginToken){ ... });
+			->
+			Identity->GetAuthToken(LocalUserNum, *MakeDelegateAdapter(this, [this](int32 UserNum, FString LoginToken){ ... });
+
+		The second instance, in addition to being slightly shorter and easier to remember, will be a unique ptr (can MoveTemp into the closure) and also will automatically do a weak ptr validity check on this
+
+	The lifetime of the adapter is a shared ptr attached to the input delegate and will live as long as that delegate is bound.
+*/
+
+namespace UE::Online {
+
+template<typename ComponentType, typename... LambdaArgs>
+class TDelegateAdapter
+	: public TSharedFromThis<TDelegateAdapter<ComponentType, LambdaArgs...>>
+{
+
+public:
+	TDelegateAdapter(TSharedRef<ComponentType>& InParent)
+		: Parent(TWeakPtr<ComponentType>(InParent))
+	{
+	}
+	
+	void SetupDelegate(TUniqueFunction<void(LambdaArgs...)>&& InCallback)
+	{
+		// callback exists outside of the delegate's bound lambda since that lambda is not move-only
+		Callback = [this, InCallback = MoveTemp(InCallback)](LambdaArgs... Args)
+		{
+			InCallback(Forward<LambdaArgs>(Args)...);
+		};
+
+	}
+
+	operator TDelegate<void(LambdaArgs...)>()
+	{
+		TSharedPtr<TDelegateAdapter<ComponentType, LambdaArgs...>> SharedAnchor = this->AsShared();
+		return TDelegate<void(LambdaArgs...)>::CreateLambda([this, SharedAnchor](LambdaArgs... Args)
+		{
+			if(Parent.IsValid())
+			{
+				Callback(Forward<LambdaArgs>(Args)...);
+			}
+		});
+	}
+
+private:
+	TWeakPtr<ComponentType> Parent;
+	TUniqueFunction<void(LambdaArgs...)> Callback;
+	FDelegateHandle Handle;
+};
+
+/*
+	These set of classes allow us to accept a generic lambda and decompose its parameters into a TUniqueFunction so we have the exact list of parameters to bind/unbind the delegate to
+*/
+namespace Private { 
+	template <typename... Params>
+	class TDAConverterHelper2
+	{
+	};
+
+	template <typename CallableObject, typename TResultType, typename... ParamTypes>
+	class TDAConverterHelper2<CallableObject, TResultType, ParamTypes...>
+	{
+	public:
+
+		TUniqueFunction<TResultType(ParamTypes...)> GetUniqueFn(CallableObject&& Callable)
+		{
+			return TUniqueFunction<TResultType(ParamTypes...)>(MoveTemp(Callable));
+		}
+
+		template<typename ComponentType>
+		TSharedPtr<TDelegateAdapter<ComponentType, ParamTypes...>> Construct(TSharedRef<ComponentType> Interface, CallableObject&& InCallback)
+		{
+ 			TSharedPtr<TDelegateAdapter<ComponentType, ParamTypes...>> Adapter = MakeShared<TDelegateAdapter<ComponentType, ParamTypes...>>(Interface);
+			Adapter->SetupDelegate(GetUniqueFn(MoveTemp(InCallback)));
+			return Adapter;
+		}
+	};
+
+	template <typename CallableObject, typename = void>
+	class TDAConverterHelper
+	{
+	};
+
+	template <typename CallableObject, typename ReturnType, typename... ParamTypes>
+	class TDAConverterHelper<CallableObject, ReturnType(ParamTypes...)>
+		: public TDAConverterHelper2<CallableObject, ReturnType, ParamTypes...>
+	{
+	};
+
+
+	template <typename CallableObject, typename ReturnType, typename ObjectType, typename... ParamTypes>
+	class TDAConverterHelper<CallableObject, ReturnType(ObjectType::*)(ParamTypes...)>
+		: public TDAConverterHelper2<CallableObject, ReturnType, ParamTypes...>
+	{
+	};
+
+	template <typename CallableObject, typename ReturnType, typename ObjectType, typename... ParamTypes>
+	class TDAConverterHelper<CallableObject, ReturnType(ObjectType::*)(ParamTypes...) const>
+		: public TDAConverterHelper2<CallableObject, ReturnType, ParamTypes...>
+	{
+	};
+
+	template <typename CallableObject>
+	class TDAConverter
+		: public TDAConverterHelper<CallableObject, decltype(&std::remove_reference_t<CallableObject>::operator())>
+	{
+	};
+} // namespace Private (UE::Online)
+
+
+template<typename ComponentType, typename Callback>
+auto MakeDelegateAdapter(TSharedRef<ComponentType> Interface, Callback&& InCallback)
+{
+	return Private::TDAConverter<Callback>().Construct(Interface, MoveTemp(InCallback));
+}
+
+template<typename ComponentType, typename Callback>
+auto MakeDelegateAdapter(TSharedPtr<ComponentType> Interface, Callback&& InCallback)
+{
+	checkf(Interface.IsValid(), TEXT("Must pass in a valid interface to MakeDelegateAdapter!"));
+	return Private::TDAConverter<Callback>().Construct(Interface.ToSharedRef(), MoveTemp(InCallback));
+}
+
+template<typename ComponentType, typename Callback>
+auto MakeDelegateAdapter(ComponentType* Interface, Callback&& InCallback)
+{
+	return Private::TDAConverter<Callback>().Construct(Interface->AsShared(), MoveTemp(InCallback));
+}
+
+/* UE::Online */ }
diff --git a/Engine/Plugins/Online/Experimental/OnlineServicesOSSAdapter/Source/Public/Online/MulticastAdapter.h b/Engine/Plugins/Online/Experimental/OnlineServicesOSSAdapter/Source/Public/Online/MulticastAdapter.h
new file mode 100644
--- /dev/null
+++ b/Engine/Plugins/Online/Experimental/OnlineServicesOSSAdapter/Source/Public/Online/MulticastAdapter.h
@@ -0,0 +1,149 @@
+// Copyright Epic Games, Inc. All Rights Reserved.
+#pragma once
+#include "OnlineSubsystemTypes.h"
+
+/*
+	BASIC USAGE
+		The multicast adapter maintains the same functionality as a delegate adapter
+
+			Identity->OnLoginCompleteDelegate.BindLambda([this](int32 UserNum){...})
+			->
+			MakeMulticastAdapter(this, Identity->OnLoginComplete, [this](int32 UserNum){...});
+
+		This type of adapter will unbind itself immediately after a single execution. 
+
+		The lifetime of the adapter is a shared ptr attached to the input delegate and will live as long as that delegate is bound.
+*/
+namespace UE::Online {
+
+template<typename ComponentType, typename DelegateType, typename... LambdaArgs>
+class TMulticastDelegateAdapter
+	: public TSharedFromThis<TMulticastDelegateAdapter<ComponentType, DelegateType, LambdaArgs...>>
+{
+
+public:
+	TMulticastDelegateAdapter(TSharedRef<ComponentType>& InParent, DelegateType& InDelegate)
+		: Parent(TWeakPtr<ComponentType>(InParent))
+		, Delegate(InDelegate)
+	{
+	}
+
+	void SetupDelegate(TUniqueFunction<void(LambdaArgs...)>&& InCallback)
+	{
+		TSharedPtr<TMulticastDelegateAdapter<ComponentType, DelegateType, LambdaArgs...>> SharedAnchor = this->AsShared();
+
+		// callback exists outside of the delegate's bound lambda since that lambda is not move-only
+		Callback = [this, InCallback = MoveTemp(InCallback)](LambdaArgs... Args)
+		{
+			InCallback(Forward<LambdaArgs>(Args)...);
+		};
+
+		Handle = Delegate.AddLambda([this, SharedAnchor](LambdaArgs... Args) mutable 
+		{
+			if(Parent.IsValid())
+			{
+				Callback(Forward<LambdaArgs>(Args)...);
+			}
+
+			SharedAnchor.Reset();
+			Delegate.Remove(Handle);
+		});
+	}
+
+	FDelegateHandle GetHandle()
+	{
+		return Handle;
+	}
+
+	TWeakPtr<TMulticastDelegateAdapter<ComponentType, DelegateType, LambdaArgs...>> AsWeak()
+	{
+		return TWeakPtr<TMulticastDelegateAdapter<ComponentType, DelegateType, LambdaArgs...>>(this->AsShared());
+	}
+
+private:
+	TWeakPtr<ComponentType> Parent;
+	TUniqueFunction<void(LambdaArgs...)> Callback;
+	DelegateType& Delegate;
+	FDelegateHandle Handle;
+};
+
+/*
+	These set of classes allow us to accept a generic lambda and decompose its parameters into a TUniqueFunction so we have the exact list of parameters to bind/unbind the delegate to
+*/
+namespace Private 
+{
+	template <typename... Params>
+	class TMAConverterHelper2
+	{
+	};
+
+	template <typename CallableObject, typename TResultType, typename... ParamTypes>
+	class TMAConverterHelper2<CallableObject, TResultType, ParamTypes...>
+	{
+	public:
+
+		TUniqueFunction<TResultType(ParamTypes...)> GetUniqueFn(CallableObject&& Callable)
+		{
+			return TUniqueFunction<TResultType(ParamTypes...)>(MoveTemp(Callable));
+		}
+
+		template<typename ComponentType, typename DelegateType>
+		TSharedPtr<TMulticastDelegateAdapter<ComponentType, DelegateType, ParamTypes...>> Construct(TSharedRef<ComponentType> Interface, DelegateType& InDelegate, CallableObject&& InCallback)
+		{
+			TSharedPtr<TMulticastDelegateAdapter<ComponentType, DelegateType, ParamTypes...>> Adapter = MakeShared<TMulticastDelegateAdapter<ComponentType, DelegateType, ParamTypes...>>(Interface, InDelegate);
+			Adapter->SetupDelegate(GetUniqueFn(MoveTemp(InCallback)));
+			return Adapter;
+		}
+	};
+
+	template <typename CallableObject, typename = void>
+	class TMAConverterHelper
+	{
+	};
+
+	template <typename CallableObject, typename ReturnType, typename... ParamTypes>
+	class TMAConverterHelper<CallableObject, ReturnType(ParamTypes...)>
+		: public TMAConverterHelper2<CallableObject, ReturnType, ParamTypes...>
+	{
+	};
+
+
+	template <typename CallableObject, typename ReturnType, typename ObjectType, typename... ParamTypes>
+	class TMAConverterHelper<CallableObject, ReturnType(ObjectType::*)(ParamTypes...)>
+		: public TMAConverterHelper2<CallableObject, ReturnType, ParamTypes...>
+	{
+	};
+
+	template <typename CallableObject, typename ReturnType, typename ObjectType, typename... ParamTypes>
+	class TMAConverterHelper<CallableObject, ReturnType(ObjectType::*)(ParamTypes...) const>
+		: public TMAConverterHelper2<CallableObject, ReturnType, ParamTypes...>
+	{
+	};
+
+	template <typename CallableObject>
+	class TMAConverter
+		: public TMAConverterHelper<CallableObject, decltype(&std::remove_reference_t<CallableObject>::operator())>
+	{
+	};
+} // namespace Private (UE::Online)
+
+template<typename ComponentType, typename DelegateType, typename Callback>
+auto MakeMulticastAdapter(TSharedRef<ComponentType> Interface, DelegateType& InDelegate, Callback&& InCallback)
+{
+	return Private::TMAConverter<Callback>().Construct(Interface, InDelegate, MoveTemp(InCallback));
+}
+
+template<typename ComponentType, typename DelegateType, typename Callback>
+auto MakeMulticastAdapter(TSharedPtr<ComponentType> Interface, DelegateType& InDelegate, Callback&& InCallback)
+{
+	checkf(Interface.IsValid(), TEXT("Must pass in a valid interface to MakeDelegateAdapter!"));
+	return Private::TMAConverter<Callback>().Construct(Interface.ToSharedRef(), InDelegate, MoveTemp(InCallback));
+}
+
+template<typename ComponentType, typename DelegateType, typename Callback>
+auto MakeMulticastAdapter(ComponentType* Interface, DelegateType& InDelegate, Callback&& InCallback)
+{
+	return Private::TMAConverter<Callback>().Construct(Interface->AsShared(), InDelegate, MoveTemp(InCallback));
+}
+
+/* UE::Online */ }
