Adapter implementation for the Presence interface

diff --git a/Engine/Plugins/Online/Experimental/OnlineServicesOSSAdapter/Source/Private/Online/AuthOSSAdapter.cpp b/Engine/Plugins/Online/Experimental/OnlineServicesOSSAdapter/Source/Private/Online/AuthOSSAdapter.cpp
--- a/Engine/Plugins/Online/Experimental/OnlineServicesOSSAdapter/Source/Private/Online/AuthOSSAdapter.cpp
+++ b/Engine/Plugins/Online/Experimental/OnlineServicesOSSAdapter/Source/Private/Online/AuthOSSAdapter.cpp
@@ -386,6 +386,11 @@ FUniqueNetIdRef FAuthOSSAdapter::GetUniqueNetId(FOnlineAccountIdHandle AccountId
 	return static_cast<FOnlineServicesOSSAdapter&>(Services).GetAccountIdRegistry().GetIdValue(AccountIdHandle);
 }
 
+FOnlineAccountIdHandle UE::Online::FAuthOSSAdapter::GetAccountIdHandle(FUniqueNetIdRef UniqueNetId) const
+{
+	return static_cast<FOnlineServicesOSSAdapter&>(Services).GetAccountIdRegistry().FindOrAddHandle(UniqueNetId);
+}
+
 int32 FAuthOSSAdapter::GetLocalUserNum(FOnlineAccountIdHandle AccountIdHandle) const
 {
 	FUniqueNetIdRef UniqueNetId = GetUniqueNetId(AccountIdHandle);
diff --git a/Engine/Plugins/Online/Experimental/OnlineServicesOSSAdapter/Source/Private/Online/OnlineServicesOSSAdapter.cpp b/Engine/Plugins/Online/Experimental/OnlineServicesOSSAdapter/Source/Private/Online/OnlineServicesOSSAdapter.cpp
--- a/Engine/Plugins/Online/Experimental/OnlineServicesOSSAdapter/Source/Private/Online/OnlineServicesOSSAdapter.cpp
+++ b/Engine/Plugins/Online/Experimental/OnlineServicesOSSAdapter/Source/Private/Online/OnlineServicesOSSAdapter.cpp
@@ -7,6 +7,7 @@
 #include "Online/ConnectivityOSSAdapter.h"
 #include "Online/PrivilegesOSSAdapter.h"
 #include "Online/ExternalUIOSSAdapter.h"
+#include "Online/PresenceOSSAdapter.h"
 
 #include "OnlineSubsystem.h"
 
@@ -24,6 +25,7 @@ void FOnlineServicesOSSAdapter::RegisterComponents()
 	Components.Register<FAuthOSSAdapter>(*this);
 	Components.Register<FConnectivityOSSAdapter>(*this);
 	Components.Register<FPrivilegesOSSAdapter>(*this);
+	Components.Register<FPresenceOSSAdapter>(*this);
 
 	if (Subsystem->GetExternalUIInterface().IsValid())
 	{
diff --git a/Engine/Plugins/Online/Experimental/OnlineServicesOSSAdapter/Source/Private/Online/PresenceOSSAdapter.cpp b/Engine/Plugins/Online/Experimental/OnlineServicesOSSAdapter/Source/Private/Online/PresenceOSSAdapter.cpp
new file mode 100644
--- /dev/null
+++ b/Engine/Plugins/Online/Experimental/OnlineServicesOSSAdapter/Source/Private/Online/PresenceOSSAdapter.cpp
@@ -0,0 +1,290 @@
+// Copyright Epic Games, Inc. All Rights Reserved.
+
+#include "Online/PresenceOSSAdapter.h"
+#include "Online/AuthOSSAdapter.h"
+
+#include "Online/OnlineServicesOSSAdapter.h"
+#include "Online/OnlineIdOSSAdapter.h"
+#include "Online/OnlineErrorDefinitions.h"
+#include "Online/DelegateAdapter.h"
+
+#include "OnlineSubsystem.h"
+#include "Interfaces/OnlineIdentityInterface.h"
+#include "Interfaces/OnlinePresenceInterface.h"
+
+
+namespace UE::Online {
+
+EUserPresenceStatus StatusV1toV2(EOnlinePresenceState::Type Status)
+{
+	switch (Status)
+	{
+	case EOnlinePresenceState::Online:
+		return EUserPresenceStatus::Online;
+		break;
+	case EOnlinePresenceState::Offline:
+		return EUserPresenceStatus::Offline;
+		break;
+	case EOnlinePresenceState::Away:
+		return EUserPresenceStatus::Away;
+		break;
+	case EOnlinePresenceState::ExtendedAway:
+		return EUserPresenceStatus::ExtendedAway;
+		break;
+	case EOnlinePresenceState::DoNotDisturb:
+		return EUserPresenceStatus::DoNotDisturb;
+		break;
+	default:
+		return EUserPresenceStatus::Unknown;
+		break;
+	}
+}
+
+EOnlinePresenceState::Type StatusV2toV1(EUserPresenceStatus Status)
+{
+	switch (Status)
+	{
+	case EUserPresenceStatus::Offline:
+		return EOnlinePresenceState::Offline;
+		break;
+	case EUserPresenceStatus::Online:
+		return EOnlinePresenceState::Online;
+		break;
+	case EUserPresenceStatus::Away:
+		return EOnlinePresenceState::Away;
+		break;
+	case EUserPresenceStatus::ExtendedAway:
+		return EOnlinePresenceState::ExtendedAway;
+		break;
+	case EUserPresenceStatus::DoNotDisturb:
+		return EOnlinePresenceState::DoNotDisturb;
+		break;
+	default:
+		return EOnlinePresenceState::Offline; // no default/unknown value for V1
+		break;
+	}
+}
+
+EUserPresenceGameStatus GameStatusV1toV2(FOnlineUserPresence& Presence)
+{
+	if (Presence.bIsPlayingThisGame)
+	{
+		return EUserPresenceGameStatus::PlayingThisGame;
+	}
+	else if (Presence.bIsPlaying)
+	{
+		return EUserPresenceGameStatus::PlayingOtherGame;
+	}
+	else 
+	{
+		return EUserPresenceGameStatus::Unknown;
+	}
+}
+
+void PropertiesV1toV2(FOnlineKeyValuePairs<FPresenceKey, FVariantData>& PropertiesIn, TMap<FString, FPresenceVariant> PropertiesOut)
+{
+	for (const TPair<FPresenceKey, FVariantData>& Pair : PropertiesIn)
+	{
+		PropertiesOut.Add(Pair.Key, Pair.Value.ToString());
+	}
+}
+
+void PropertiesV2toV1(TMap<FString, FPresenceVariant> PropertiesIn, FOnlineKeyValuePairs<FPresenceKey, FVariantData>& PropertiesOut)
+{
+	for (const TPair<FString, FPresenceVariant>& Pair : PropertiesIn)
+	{
+		FVariantData Data;
+		Data.SetValue(Pair.Value);
+		PropertiesOut.Add(Pair.Key, Data);
+	}
+}
+
+TSharedRef<FUserPresence> FPresenceOSSAdapter::PresenceV1toV2(FOnlineUserPresence& Presence)
+{
+	TSharedRef<FUserPresence> PresenceV2 = MakeShared<FUserPresence>();
+
+	PresenceV2->UserId = Auth->GetAccountIdHandle(Presence.SessionId.ToSharedRef());
+	PresenceV2->Status = StatusV1toV2(Presence.Status.State);
+	PresenceV2->GameStatus = GameStatusV1toV2(Presence);
+	PresenceV2->Joinability = (Presence.bIsJoinable) ? EUserPresenceJoinability::Public : EUserPresenceJoinability::Private;
+	PresenceV2->StatusString = Presence.Status.StatusStr;
+	PresenceV2->RichPresenceString = Presence.Status.StatusStr;
+	PropertiesV1toV2(Presence.Status.Properties, PresenceV2->Properties);
+
+	return PresenceV2;
+}
+
+TSharedRef<FOnlineUserPresence> FPresenceOSSAdapter::PresenceV2toV1(FUserPresence& Presence)
+{
+	TSharedRef<FOnlineUserPresence> PresenceV1 = MakeShared<FOnlineUserPresence>();
+
+	PresenceV1->SessionId = Auth->GetUniqueNetId(Presence.UserId);
+	PresenceV1->bIsOnline = (Presence.Status == EUserPresenceStatus::Online) ? 1 : 0;
+	PresenceV1->bIsPlaying = (Presence.GameStatus != EUserPresenceGameStatus::Unknown) ? 1 : 0;
+	PresenceV1->bIsPlayingThisGame = (Presence.GameStatus == EUserPresenceGameStatus::PlayingThisGame) ? 1 : 0;
+	PresenceV1->bIsJoinable = (Presence.Joinability == EUserPresenceJoinability::Public) ? 1 : 0;
+	PresenceV1->Status.StatusStr = Presence.StatusString;
+	PresenceV1->Status.State = StatusV2toV1(Presence.Status);
+	PropertiesV2toV1(Presence.Properties, PresenceV1->Status.Properties);
+
+	return PresenceV1;
+}
+
+void FPresenceOSSAdapter::PostInitialize()
+{
+	Auth = Services.Get<FAuthOSSAdapter>();
+	PresenceInt = static_cast<FOnlineServicesOSSAdapter&>(Services).GetSubsystem().GetPresenceInterface();
+}
+
+void FPresenceOSSAdapter::PreShutdown()
+{
+}
+
+TOnlineAsyncOpHandle<FQueryPresence> FPresenceOSSAdapter::QueryPresence(FQueryPresence::Params&& Params)
+{
+	TOnlineAsyncOpRef<FQueryPresence> Op = GetOp<FQueryPresence>(MoveTemp(Params));
+
+	Op->Then([this](TOnlineAsyncOp<FQueryPresence>& Op, TPromise<TOnlineResult<FQueryPresence>>&& Result)
+	{
+		const FUniqueNetIdRef UniqueNetId = Auth->GetUniqueNetId(Op.GetParams().LocalUserId);
+
+		if (UniqueNetId->IsValid())
+		{
+			PresenceInt->QueryPresence(*UniqueNetId, *MakeDelegateAdapter(this, [this, ResultPromise = MoveTemp(Result)](const FUniqueNetId& UserId, const bool bWasSuccessful) mutable
+			{
+				if (bWasSuccessful)
+				{
+					FQueryPresence::Result Result;
+					ResultPromise.EmplaceValue(MoveTemp(Result));
+				}
+				else
+				{
+					ResultPromise.EmplaceValue(Errors::Unknown());
+				}
+			}));
+		}
+		else
+		{
+			Result.EmplaceValue(Errors::InvalidUser());
+		}
+	});
+
+	return Op->GetHandle();
+}
+
+// Base implementation for the batch query function makes multiple async calls to query and returns them as a single output
+TOnlineAsyncOpHandle<FBatchQueryPresence> FPresenceOSSAdapter::BatchQueryPresence(FBatchQueryPresence::Params&& Params)
+{
+	TOnlineAsyncOpRef<FBatchQueryPresence> Op = GetOp<FBatchQueryPresence>(MoveTemp(Params));
+
+	Op->Then([this](TOnlineAsyncOp<FBatchQueryPresence>& Op, TPromise<TOnlineResult<FBatchQueryPresence>>&& Result)
+	{
+		const FUniqueNetIdRef LocalUserId = Auth->GetUniqueNetId(Op.GetParams().LocalUserId);
+		TArray<FUniqueNetIdRef> NetIds;
+		for (const FOnlineAccountIdHandle& TargetUserId : Op.GetParams().TargetUserIds)
+		{
+			NetIds.Add(Auth->GetUniqueNetId(TargetUserId));
+		}
+
+		PresenceInt->QueryPresence(*LocalUserId, NetIds, *MakeDelegateAdapter(this, [this, ResultPromise=MoveTemp(Result)](const FUniqueNetId& UserId, const bool bWasSuccessful) mutable
+		{
+			if(bWasSuccessful)
+			{
+				FBatchQueryPresence::Result Result;
+				ResultPromise.EmplaceValue(MoveTemp(Result));
+			}
+			else 
+			{
+				ResultPromise.EmplaceValue(Errors::Unknown());
+			}
+		}));
+	});
+
+	return Op->GetHandle();
+}
+
+TOnlineResult<FGetCachedPresence> FPresenceOSSAdapter::GetCachedPresence(FGetCachedPresence::Params&& Params)
+{
+	const FUniqueNetId& User = *Auth->GetUniqueNetId(Params.LocalUserId);
+	TSharedPtr<FOnlineUserPresence> Presence = MakeShared<FOnlineUserPresence>();
+	EOnlineCachedResult::Type Result = PresenceInt->GetCachedPresence(User, Presence);
+	if (Result != EOnlineCachedResult::Success)
+	{
+		TSharedRef<FUserPresence> V2Presence = PresenceV1toV2(*Presence);
+		return TOnlineResult<FGetCachedPresence>({V2Presence});
+	}
+	else
+	{
+		return TOnlineResult<FGetCachedPresence>(Errors::NotFound());
+	}
+}
+
+TOnlineAsyncOpHandle<FUpdatePresence> FPresenceOSSAdapter::UpdatePresence(FUpdatePresence::Params&& Params)
+{
+	TOnlineAsyncOpRef<FUpdatePresence> Op = GetOp<FUpdatePresence>(MoveTemp(Params));
+
+	Op->Then([this](TOnlineAsyncOp<FUpdatePresence>& Op, TPromise<TOnlineResult<FUpdatePresence>>&& Result)
+	{
+		TSharedRef<const FOnlineUserPresence> PresenceV1 = PresenceV2toV1(*Op.GetParams().Presence);
+		PresenceInt->SetPresence(*Auth->GetUniqueNetId(Op.GetParams().LocalUserId), PresenceV1->Status, *MakeDelegateAdapter(this, [this, ResultPromise = MoveTemp(Result)](const FUniqueNetId& UserId, const bool bWasSuccessful) mutable
+		{
+			if (bWasSuccessful)
+			{
+				FUpdatePresence::Result Result;
+				ResultPromise.EmplaceValue(MoveTemp(Result));
+			}
+			else
+			{
+				ResultPromise.EmplaceValue(Errors::Unknown());
+			}
+		}));
+	});
+
+	return Op->GetHandle();
+}
+
+TOnlineAsyncOpHandle<FPartialUpdatePresence> FPresenceOSSAdapter::PartialUpdatePresence(FPartialUpdatePresence::Params&& Params)
+{
+	TOnlineAsyncOpRef<FPartialUpdatePresence> Op = GetOp<FPartialUpdatePresence>(MoveTemp(Params));
+
+	Op->Then([this](TOnlineAsyncOp<FPartialUpdatePresence>& Op, TPromise<TOnlineResult<FPartialUpdatePresence>>&& Result)
+	{
+		const FPartialUpdatePresence::Params::FMutations& Mutations = Op.GetParams().Mutations;
+		FOnlinePresenceSetPresenceParameters Parameters;
+	
+		if (Mutations.StatusString.IsSet())
+		{
+			Parameters.StatusStr = Mutations.StatusString.GetValue();
+		}
+
+		if (Mutations.Status.IsSet())
+		{
+			Parameters.State = StatusV2toV1(Mutations.Status.GetValue());
+		}
+
+		// todo: V1 doesnt have a clear path for doing key removals
+		if (Mutations.UpdatedProperties.Num() > 0)
+		{
+			FOnlineKeyValuePairs<FPresenceKey, FVariantData> NewProperties;
+			PropertiesV2toV1(Mutations.UpdatedProperties, NewProperties);
+			Parameters.Properties = NewProperties;
+		}
+
+		PresenceInt->SetPresence(*Auth->GetUniqueNetId(Op.GetParams().LocalUserId), MoveTemp(Parameters), *MakeDelegateAdapter(this, [this, ResultPromise = MoveTemp(Result)](const FUniqueNetId& UserId, const bool bWasSuccessful) mutable
+		{
+			if (bWasSuccessful)
+			{
+				FPartialUpdatePresence::Result Result;
+				ResultPromise.EmplaceValue(MoveTemp(Result));
+			}
+			else
+			{
+				ResultPromise.EmplaceValue(Errors::Unknown());
+			}
+		}));
+	});
+
+	return Op->GetHandle();
+}
+
+/* UE::Online */ }
diff --git a/Engine/Plugins/Online/Experimental/OnlineServicesOSSAdapter/Source/Public/Online/AuthOSSAdapter.h b/Engine/Plugins/Online/Experimental/OnlineServicesOSSAdapter/Source/Public/Online/AuthOSSAdapter.h
--- a/Engine/Plugins/Online/Experimental/OnlineServicesOSSAdapter/Source/Public/Online/AuthOSSAdapter.h
+++ b/Engine/Plugins/Online/Experimental/OnlineServicesOSSAdapter/Source/Public/Online/AuthOSSAdapter.h
@@ -31,6 +31,7 @@ public:
 	virtual TOnlineResult<FAuthGetAccountByAccountId> GetAccountByAccountId(FAuthGetAccountByAccountId::Params&& Params) override;
 
 	FUniqueNetIdRef GetUniqueNetId(FOnlineAccountIdHandle AccountIdHandle) const;
+	FOnlineAccountIdHandle GetAccountIdHandle(FUniqueNetIdRef UniqueNetId) const;
 	int32 GetLocalUserNum(FOnlineAccountIdHandle AccountIdHandle) const;
 
 	IOnlineIdentityPtr GetIdentityInterface() const;
diff --git a/Engine/Plugins/Online/Experimental/OnlineServicesOSSAdapter/Source/Public/Online/PresenceOSSAdapter.h b/Engine/Plugins/Online/Experimental/OnlineServicesOSSAdapter/Source/Public/Online/PresenceOSSAdapter.h
new file mode 100644
--- /dev/null
+++ b/Engine/Plugins/Online/Experimental/OnlineServicesOSSAdapter/Source/Public/Online/PresenceOSSAdapter.h
@@ -0,0 +1,43 @@
+// Copyright Epic Games, Inc. All Rights Reserved.
+
+#pragma once
+
+#include "Online/PresenceCommon.h"
+#include "AuthOSSAdapter.h"
+#include "OnlineSubsystemTypes.h"
+
+class IOnlinePresence;
+using IOnlinePresencePtr = TSharedPtr<IOnlinePresence>;
+
+namespace UE::Online {
+
+class FPresenceOSSAdapter : public FPresenceCommon
+{
+public:
+	using Super = FPresenceCommon;
+
+	using FPresenceCommon::FPresenceCommon;
+
+	// IOnlineComponent
+	virtual void PostInitialize() override;
+	virtual void PreShutdown() override;
+
+	// IPresence
+	virtual TOnlineAsyncOpHandle<FQueryPresence> QueryPresence(FQueryPresence::Params&& Params) override;
+	virtual TOnlineAsyncOpHandle<FBatchQueryPresence> BatchQueryPresence(FBatchQueryPresence::Params&& Params) override;
+	virtual TOnlineResult<FGetCachedPresence> GetCachedPresence(FGetCachedPresence::Params&& Params) override;
+	virtual TOnlineAsyncOpHandle<FUpdatePresence> UpdatePresence(FUpdatePresence::Params&& Params) override;
+	virtual TOnlineAsyncOpHandle<FPartialUpdatePresence> PartialUpdatePresence(FPartialUpdatePresence::Params&& Params) override;
+
+	FUniqueNetIdRef GetUniqueNetId(FOnlineAccountIdHandle AccountIdHandle) const;
+	int32 GetLocalUserNum(FOnlineAccountIdHandle AccountIdHandle) const;
+
+protected:
+	const FAuthOSSAdapter* Auth;
+	IOnlinePresencePtr PresenceInt;
+
+	TSharedRef<FUserPresence> PresenceV1toV2(FOnlineUserPresence& Presence);
+	TSharedRef<FOnlineUserPresence> PresenceV2toV1(FUserPresence& Presence);
+};
+
+/* UE::Online */ }
