OSSv2: Add Stats OSSv1 Adaptor.

diff --git a/Engine/Plugins/Online/Experimental/OnlineServicesOSSAdapter/Source/Private/Online/OnlineServicesOSSAdapter.cpp b/Engine/Plugins/Online/Experimental/OnlineServicesOSSAdapter/Source/Private/Online/OnlineServicesOSSAdapter.cpp
--- a/Engine/Plugins/Online/Experimental/OnlineServicesOSSAdapter/Source/Private/Online/OnlineServicesOSSAdapter.cpp
+++ b/Engine/Plugins/Online/Experimental/OnlineServicesOSSAdapter/Source/Private/Online/OnlineServicesOSSAdapter.cpp
@@ -8,6 +8,7 @@
 #include "Online/PrivilegesOSSAdapter.h"
 #include "Online/ExternalUIOSSAdapter.h"
 #include "Online/PresenceOSSAdapter.h"
+#include "Online/StatsOSSAdapter.h"
 
 #include "OnlineSubsystem.h"
 
@@ -26,6 +27,7 @@ void FOnlineServicesOSSAdapter::RegisterComponents()
 	Components.Register<FConnectivityOSSAdapter>(*this);
 	Components.Register<FPrivilegesOSSAdapter>(*this);
 	Components.Register<FPresenceOSSAdapter>(*this);
+	Components.Register<FStatsOSSAdapter>(*this);
 
 	if (Subsystem->GetExternalUIInterface().IsValid())
 	{
diff --git a/Engine/Plugins/Online/Experimental/OnlineServicesOSSAdapter/Source/Private/Online/StatsOSSAdapter.cpp b/Engine/Plugins/Online/Experimental/OnlineServicesOSSAdapter/Source/Private/Online/StatsOSSAdapter.cpp
new file mode 100644
--- /dev/null
+++ b/Engine/Plugins/Online/Experimental/OnlineServicesOSSAdapter/Source/Private/Online/StatsOSSAdapter.cpp
@@ -0,0 +1,234 @@
+// Copyright Epic Games, Inc. All Rights Reserved.
+
+#include "Online/StatsOSSAdapter.h"
+#include "Online/AuthOSSAdapter.h"
+
+#include "Online/OnlineServicesOSSAdapter.h"
+#include "Online/OnlineIdOSSAdapter.h"
+#include "Online/OnlineErrorDefinitions.h"
+#include "Online/DelegateAdapter.h"
+
+#include "OnlineSubsystem.h"
+#include "Interfaces/OnlineIdentityInterface.h"
+
+namespace UE::Online {
+
+void FStatsOSSAdapter::PostInitialize()
+{
+	Super::PostInitialize();
+
+	Auth = Services.Get<FAuthOSSAdapter>();
+	StatsInterface = static_cast<FOnlineServicesOSSAdapter&>(Services).GetSubsystem().GetStatsInterface();
+}
+
+void FStatsOSSAdapter::ConvertStatValueV2ToStatUpdateV1(const FString& StatName, const FStatValue& StatValue, FOnlineStatUpdate& OutOnlineStatUpdate)
+{
+	FOnlineStatValue OnlineStatValue;
+
+	// TODO: Add the FOnlineStatValue conversion when SchemaVariant is ready
+
+	FOnlineStatUpdate::EOnlineStatModificationType OnlineStatModificationType = FOnlineStatUpdate::EOnlineStatModificationType::Unknown;
+
+	if (FStatDefinition* StatDefinition = StatDefinitions.Find(StatName))
+	{
+		switch (StatDefinition->ModifyMethod)
+		{
+		case EStatModifyMethod::Sum: OnlineStatModificationType = FOnlineStatUpdate::EOnlineStatModificationType::Sum; break;
+		case EStatModifyMethod::Set: OnlineStatModificationType = FOnlineStatUpdate::EOnlineStatModificationType::Set; break;
+		case EStatModifyMethod::Largest: OnlineStatModificationType = FOnlineStatUpdate::EOnlineStatModificationType::Largest; break;
+		case EStatModifyMethod::Smallest: OnlineStatModificationType = FOnlineStatUpdate::EOnlineStatModificationType::Smallest; break;
+		}
+	}
+	else
+	{
+		UE_LOG(LogTemp, Warning, TEXT("Can't find stat definition %s"), *StatName);
+	}
+
+	OutOnlineStatUpdate.Set(OnlineStatValue, OnlineStatModificationType);
+}
+
+void FStatsOSSAdapter::ConvertStatValueV1ToStatValueV2(const FOnlineStatValue& OnlineStatValue, FStatValue& OutStatValue)
+{
+	// TODO: Add the conversion when SchemaVariant is ready
+}
+
+FOnlineError FStatsOSSAdapter::ConvertUpdateUsersStatsV2toV1(const TArray<FUserStats>& UpdateUsersStats, TArray<FOnlineStatsUserUpdatedStats>& OutUpdateUserStatsV1)
+{
+	for (const FUserStats& UserStats : UpdateUsersStats)
+	{
+		const FUniqueNetIdRef UserId = Auth->GetUniqueNetId(UserStats.UserId);
+		if (!UserId->IsValid())
+		{
+			return Errors::InvalidUser();
+		}
+
+		FOnlineStatsUserUpdatedStats OnlineStatsUserUpdatedStats(UserId);
+		for (const TPair<FString, FStatValue>& StatPair : UserStats.Stats)
+		{
+			FOnlineStatUpdate OnlineStatUpdate;
+			ConvertStatValueV2ToStatUpdateV1(StatPair.Key, StatPair.Value, OnlineStatUpdate);
+			OnlineStatsUserUpdatedStats.Stats.Emplace(StatPair.Key, OnlineStatUpdate);
+		}
+
+		OutUpdateUserStatsV1.Emplace(MoveTemp(OnlineStatsUserUpdatedStats));
+	}
+
+	return Errors::Success();
+}
+
+TOnlineAsyncOpHandle<FUpdateStats> FStatsOSSAdapter::UpdateStats(FUpdateStats::Params&& Params)
+{
+	TOnlineAsyncOpRef<FUpdateStats> Op = GetOp<FUpdateStats>(MoveTemp(Params));
+
+	Op->Then([this](TOnlineAsyncOp<FUpdateStats>& Op)
+	{
+		const FUniqueNetIdRef LocalUserId = Auth->GetUniqueNetId(Op.GetParams().LocalUserId);
+		if (!LocalUserId->IsValid())
+		{
+			Op.SetError(Errors::InvalidUser());
+			return;
+		}
+
+		TArray<FOnlineStatsUserUpdatedStats> UpdateUserStatsV1;
+		FOnlineError OnlineErrorConvert = ConvertUpdateUsersStatsV2toV1(Op.GetParams().UpdateUsersStats, UpdateUserStatsV1);
+		if (OnlineErrorConvert != Errors::Success())
+		{
+			Op.SetError(MoveTemp(OnlineErrorConvert));
+			return;
+		}
+
+		StatsInterface->UpdateStats(LocalUserId, UpdateUserStatsV1, *MakeDelegateAdapter(this, [WeakOp = Op.AsWeak()](const ::FOnlineError& OnlineError) mutable
+		{
+			if (TSharedPtr<TOnlineAsyncOp<FUpdateStats>> PinnedOp = WeakOp.Pin())
+			{
+				if (!OnlineError.WasSuccessful())
+				{
+					// TODO: Convert the error when FOnlineError conversion is ready
+					PinnedOp->SetError(Errors::Unknown());
+					return;
+				}
+
+				PinnedOp->SetResult({});
+			}
+		}));
+	}) 
+	.Enqueue(GetSerialQueue());
+
+	return Op->GetHandle();
+}
+
+TOnlineAsyncOpHandle<FQueryStats> FStatsOSSAdapter::QueryStats(FQueryStats::Params&& Params)
+{
+	TOnlineAsyncOpRef<FQueryStats> Op = GetOp<FQueryStats>(MoveTemp(Params));
+
+	Op->Then([this](TOnlineAsyncOp<FQueryStats>& Op)
+	{
+		const FUniqueNetIdRef LocalUniqueNetId = Auth->GetUniqueNetId(Op.GetParams().LocalUserId);
+		const FUniqueNetIdRef TargetUniqueNetId = Auth->GetUniqueNetId(Op.GetParams().TargetUserId);
+
+		if (!LocalUniqueNetId->IsValid() || TargetUniqueNetId->IsValid())
+		{
+			Op.SetError(Errors::InvalidUser());
+			return;
+		}
+
+		StatsInterface->QueryStats(LocalUniqueNetId, TargetUniqueNetId, *MakeDelegateAdapter(this, [this, WeakOp = Op.AsWeak()](const ::FOnlineError& OnlineError, const TSharedPtr<const FOnlineStatsUserStats>& QueriedStats) mutable
+		{
+			if (TSharedPtr<TOnlineAsyncOp<FQueryStats>> PinnedOp = WeakOp.Pin())
+			{
+				if (!OnlineError.WasSuccessful())
+				{
+					// TODO: Convert the error when FOnlineError conversion is ready
+					PinnedOp->SetError(Errors::Unknown()); // TODO: Convert the error?
+					return;
+				}
+
+				FQueryStats::Result Result;
+				for (const TPair<FString, FOnlineStatValue>& StatPair : QueriedStats->Stats)
+				{
+					FStatValue StatValue;
+					ConvertStatValueV1ToStatValueV2(StatPair.Value, StatValue);
+					Result.Stats.Emplace(StatPair.Key, StatValue);
+				}
+
+				PinnedOp->SetResult(MoveTemp(Result));
+			}
+		}));
+	})
+	.Enqueue(GetSerialQueue());
+
+	return Op->GetHandle();
+}
+
+TOnlineAsyncOpHandle<FBatchQueryStats> FStatsOSSAdapter::BatchQueryStats(FBatchQueryStats::Params&& Params)
+{
+	TOnlineAsyncOpRef<FBatchQueryStats> Op = GetOp<FBatchQueryStats>(MoveTemp(Params));
+
+	Op->Then([this](TOnlineAsyncOp<FBatchQueryStats>& Op)
+	{
+		const FUniqueNetIdRef LocalUserId = Auth->GetUniqueNetId(Op.GetParams().LocalUserId);
+		TArray<FUniqueNetIdRef> NetIds;
+		for (const FOnlineAccountIdHandle& TargetUserId : Op.GetParams().TargetUserIds)
+		{
+			const FUniqueNetIdRef UniqueNetId = Auth->GetUniqueNetId(Op.GetParams().LocalUserId);
+			if (!UniqueNetId->IsValid())
+			{
+				Op.SetError(Errors::InvalidUser());
+				return;
+			}
+
+			NetIds.Add(UniqueNetId);
+		}
+
+		StatsInterface->QueryStats(LocalUserId, NetIds, Op.GetParams().StatNames, *MakeDelegateAdapter(this, [this, WeakOp = Op.AsWeak()](const ::FOnlineError& OnlineError, const TArray<TSharedRef<const FOnlineStatsUserStats>>& UsersStatsResult) mutable
+		{
+			if (TSharedPtr<TOnlineAsyncOp<FBatchQueryStats>> PinnedOp = WeakOp.Pin())
+			{
+				if (!OnlineError.WasSuccessful())
+				{
+					// TODO: Convert the error when FOnlineError conversion is ready
+					PinnedOp->SetError(Errors::Unknown());
+				}
+
+				FBatchQueryStats::Result Result;
+
+				for (const TSharedRef<const FOnlineStatsUserStats>& OnlineStatsUserStats : UsersStatsResult)
+				{
+					FUserStats& UserStats = Result.UsersStats.Emplace_GetRef();
+					UserStats.UserId = Auth->GetAccountIdHandle(OnlineStatsUserStats->Account);
+
+					for (const TPair<FString, FOnlineStatValue>& StatPair : OnlineStatsUserStats->Stats)
+					{
+						FStatValue StatValue;
+						ConvertStatValueV1ToStatValueV2(StatPair.Value, StatValue);
+						UserStats.Stats.Emplace(StatPair.Key, StatValue);
+					}
+				}
+
+				PinnedOp->SetResult(MoveTemp(Result));
+			}
+		}));
+	})
+	.Enqueue(GetSerialQueue());
+
+	return Op->GetHandle();
+}
+
+#if !UE_BUILD_SHIPPING
+TOnlineAsyncOpHandle<FResetStats> FStatsOSSAdapter::ResetStats(FResetStats::Params&& Params)
+{
+	TOnlineAsyncOpRef<FResetStats> Op = GetOp<FResetStats>(MoveTemp(Params));
+
+	Op->Then([this](TOnlineAsyncOp<FResetStats>& Op)
+	{
+		const FUniqueNetIdRef LocalUserId = Auth->GetUniqueNetId(Op.GetParams().LocalUserId);
+		StatsInterface->ResetStats(LocalUserId);
+		Op.SetResult({});
+	})
+	.Enqueue(GetSerialQueue());
+
+	return Op->GetHandle();
+}
+#endif // !UE_BUILD_SHIPPING
+
+/* UE::Online */ }
diff --git a/Engine/Plugins/Online/Experimental/OnlineServicesOSSAdapter/Source/Public/Online/StatsOSSAdapter.h b/Engine/Plugins/Online/Experimental/OnlineServicesOSSAdapter/Source/Public/Online/StatsOSSAdapter.h
new file mode 100644
--- /dev/null
+++ b/Engine/Plugins/Online/Experimental/OnlineServicesOSSAdapter/Source/Public/Online/StatsOSSAdapter.h
@@ -0,0 +1,42 @@
+// Copyright Epic Games, Inc. All Rights Reserved.
+
+#pragma once
+
+#include "Online/StatsCommon.h"
+#include "AuthOSSAdapter.h"
+#include "OnlineSubsystemTypes.h"
+#include "Interfaces/OnlineStatsInterface.h"
+
+class IOnlineStats;
+using IOnlineStatsPtr = TSharedPtr<IOnlineStats>;
+
+namespace UE::Online {
+
+class FStatsOSSAdapter : public FStatsCommon
+{
+public:
+	using Super = FStatsCommon;
+
+	using FStatsCommon::FStatsCommon;
+
+	// IOnlineComponent
+	virtual void PostInitialize() override;
+
+	// IStats
+	virtual TOnlineAsyncOpHandle<FUpdateStats> UpdateStats(FUpdateStats::Params&& Params) override;
+	virtual TOnlineAsyncOpHandle<FQueryStats> QueryStats(FQueryStats::Params&& Params) override;
+	virtual TOnlineAsyncOpHandle<FBatchQueryStats> BatchQueryStats(FBatchQueryStats::Params&& Params) override;
+#if !UE_BUILD_SHIPPING
+	virtual TOnlineAsyncOpHandle<FResetStats> ResetStats(FResetStats::Params&& Params) override;
+#endif // !UE_BUILD_SHIPPING
+
+private:
+	FOnlineError ConvertUpdateUsersStatsV2toV1(const TArray<FUserStats>& UpdateUsersStats, TArray<FOnlineStatsUserUpdatedStats>& OutUpdateUserStatsV1);
+	void ConvertStatValueV2ToStatUpdateV1(const FString& StatName, const FStatValue& StatValue, FOnlineStatUpdate& OutOnlineStatUpdate);
+	void ConvertStatValueV1ToStatValueV2(const FOnlineStatValue& OnlineStatValue, FStatValue& OutStatValue);
+
+	const FAuthOSSAdapter* Auth;
+	IOnlineStatsPtr StatsInterface;
+};
+
+/* UE::Online */ }
