Added the ErrorsOSSAdapter header for converting OSSv1's FOnlineError class into v2's FOnlineError classes. Also includes some functionality for a platform to hook in and specify what some of their platform-specific errors should map to.

Additionally includes a small change for GetLogString of Errors to also do nested iteration on the log string of the inner errors, so we can have log strings like "Not implemented (com.epicgames.eos.missing_header)". This change is bundled with the adapter work as it is necessary for FailExtended errors in Ossv1 to present their actual extended information to the log.

diff --git a/Engine/Plugins/Online/Experimental/OnlineServices/Source/OnlineServicesInterface/Public/Online/OnlineError.h b/Engine/Plugins/Online/Experimental/OnlineServices/Source/OnlineServicesInterface/Public/Online/OnlineError.h
--- a/Engine/Plugins/Online/Experimental/OnlineServices/Source/OnlineServicesInterface/Public/Online/OnlineError.h
+++ b/Engine/Plugins/Online/Experimental/OnlineServices/Source/OnlineServicesInterface/Public/Online/OnlineError.h
@@ -75,8 +75,11 @@ public:
 		}
 	}
 
-	FString GetLogString(bool bIncludePrefix = true) const
+	FString GetLogString(bool bIncludePrefix = true, bool bIncludeSuccess = true) const
 	{
+#if !NO_LOGGING
+		FString MyLogString = TEXT("");
+
 		if (Details)
 		{
 			// likely also print out the error code first, with Details->GetLogString only providing additional information like a human readable string
@@ -87,16 +90,37 @@ public:
 			}
 
 			FString LogString = Details->GetLogString(*this);
-			return Details->GetLogString(*this);
+			MyLogString = Details->GetLogString(*this);
 		}
 		else if (ErrorCode == Errors::ErrorCode::Success)
 		{
-			return TEXT("Success");
+			if (bIncludeSuccess)
+			{
+				MyLogString = TEXT("Success");
+			}
 		}
 		else
 		{
-			return FString::Printf(TEXT("Error %x"), ErrorCode); // would use %s here and below if we were to use a string ErrorCodeType
+			MyLogString = FString::Printf(TEXT("%x"), ErrorCode); // would use %s here and below if we were to use a string ErrorCodeType
 		}
+
+		if (GetInner() != nullptr)
+		{
+			FString InnerLogStr = GetInner()->GetLogString(false, false);
+			if (!InnerLogStr.IsEmpty())
+			{
+				return FString::Printf(TEXT("%s (%s)"), *MyLogString, *InnerLogStr);
+			}
+			else
+			{
+				return MyLogString;
+			}
+		}
+
+		return MyLogString;
+#else
+		return TEXT("");
+#endif
 	}
 
 	FString GetErrorId() const
diff --git a/Engine/Plugins/Online/Experimental/OnlineServicesOSSAdapter/Source/Private/Online/ErrorsOSSAdapter.cpp b/Engine/Plugins/Online/Experimental/OnlineServicesOSSAdapter/Source/Private/Online/ErrorsOSSAdapter.cpp
new file mode 100644
--- /dev/null
+++ b/Engine/Plugins/Online/Experimental/OnlineServicesOSSAdapter/Source/Private/Online/ErrorsOSSAdapter.cpp
@@ -0,0 +1,97 @@
+// Copyright Epic Games, Inc. All Rights Reserved.
+#include "Online/ErrorsOSSAdapter.h"
+#include "Online/OnlineErrorDefinitions.h"
+#include "OnlineSubsystemTypes.h"
+#include "OnlineError.h"
+
+namespace UE::Online::Errors {
+
+TArray<TOssPlatformErrorHandler> PlatformFunctions;
+void AddOssPlatformErrorHandler(TOssPlatformErrorHandler&& InFunction)
+{
+	PlatformFunctions.Add(MoveTemp(InFunction));
+}
+
+
+inline FOnlineError Internal_OssGetCommonError(FOnlineError Error, EOnlineErrorResult Result)
+{
+	switch (Result)
+	{
+	case EOnlineErrorResult::Success: return Errors::Success(Error);
+	case EOnlineErrorResult::NoConnection: return Errors::NoConnection(Error);
+	case EOnlineErrorResult::RequestFailure: return Errors::RequestFailure(Error);
+	case EOnlineErrorResult::InvalidCreds: return Errors::InvalidCreds(Error);
+	case EOnlineErrorResult::InvalidUser: return Errors::InvalidUser(Error);
+	case EOnlineErrorResult::InvalidAuth: return Errors::InvalidAuth(Error);
+	case EOnlineErrorResult::AccessDenied: return Errors::AccessDenied(Error);
+	case EOnlineErrorResult::TooManyRequests: return Errors::TooManyRequests(Error);
+	case EOnlineErrorResult::AlreadyPending: return Errors::AlreadyPending(Error);
+	case EOnlineErrorResult::InvalidParams: return Errors::InvalidParams(Error);
+	case EOnlineErrorResult::CantParse: return Errors::CantParse(Error);
+	case EOnlineErrorResult::InvalidResults: return Errors::InvalidResults(Error);
+	case EOnlineErrorResult::IncompatibleVersion: return Errors::IncompatibleVersion(Error);
+	case EOnlineErrorResult::NotConfigured: return Errors::NotConfigured(Error);
+	case EOnlineErrorResult::NotImplemented: return Errors::NotImplemented(Error);
+	case EOnlineErrorResult::MissingInterface: return Errors::MissingInterface(Error);
+	case EOnlineErrorResult::Canceled: return Errors::Cancelled(Error);
+	default: return Error;
+	}
+}
+
+inline FOnlineError Internal_OssWrapInner(FOnlineError Error, const FOnlineErrorOss& Result)
+{
+	if (Result.WasSuccessful())
+	{
+		return Errors::Success(Error);
+	}
+
+	for (const TOssPlatformErrorHandler& PlatformHandler : PlatformFunctions)
+	{
+		TOptional<FOnlineError> NewError = PlatformHandler(Error, Result);
+		if (NewError.IsSet())
+		{
+			return Internal_OssGetCommonError(NewError.GetValue(), Result.GetErrorResult());
+		}
+	}
+
+	return Internal_OssGetCommonError(Error, Result.GetErrorResult());
+}
+
+inline FOnlineError FromOssError(const FOnlineErrorOss& Result)
+{
+	FString ErrorCode;
+	FText ErrorMessage = Result.GetErrorMessage();
+
+	if (Result.GetErrorRaw().Len() > 0)
+	{
+		ErrorCode = FString::Printf(TEXT("{code: %s, raw: %s}"), *Result.GetErrorCode(), *Result.GetErrorRaw());
+	}
+	else
+	{
+		ErrorCode = Result.GetErrorCode();
+	}
+
+	return Internal_OssWrapInner(FOnlineError(ErrorCode::Create(ErrorCode::Category::Oss_System, ErrorCode::Category::Oss, (uint32)Result.GetErrorResult()), MakeShared<FOnlineErrorDetails, ESPMode::ThreadSafe>(MoveTemp(ErrorCode), MoveTemp(ErrorMessage)), nullptr), Result);
+}
+
+} //namespace UE::Online::Errors
+
+ONLINESERVICESOSSADAPTER_API bool operator==(const UE::Online::FOnlineError& Left, const FOnlineErrorOss& Right)
+{
+	return Left == UE::Online::Errors::FromOssError(Right);
+}
+
+ONLINESERVICESOSSADAPTER_API bool operator==(const FOnlineErrorOss& Left, const UE::Online::FOnlineError& Right)
+{
+	return Right == UE::Online::Errors::FromOssError(Left);
+}
+
+ONLINESERVICESOSSADAPTER_API bool operator!=(const UE::Online::FOnlineError& Left, const FOnlineErrorOss& Right)
+{
+	return !(Left == UE::Online::Errors::FromOssError(Right));
+}
+
+ONLINESERVICESOSSADAPTER_API bool operator!=(const FOnlineErrorOss& Left, const UE::Online::FOnlineError& Right)
+{
+	return !(Right == UE::Online::Errors::FromOssError(Left));
+}
\ No newline at end of file
diff --git a/Engine/Plugins/Online/Experimental/OnlineServicesOSSAdapter/Source/Public/Online/ErrorsOSSAdapter.h b/Engine/Plugins/Online/Experimental/OnlineServicesOSSAdapter/Source/Public/Online/ErrorsOSSAdapter.h
new file mode 100644
--- /dev/null
+++ b/Engine/Plugins/Online/Experimental/OnlineServicesOSSAdapter/Source/Public/Online/ErrorsOSSAdapter.h
@@ -0,0 +1,24 @@
+// Copyright Epic Games, Inc. All Rights Reserved.
+#pragma once
+#include "Online/OnlineError.h"
+
+struct FOnlineError;
+using FOnlineErrorOss = ::FOnlineError;
+
+namespace UE::Online::Errors {
+
+	UE_ONLINE_ERROR_CATEGORY(Oss, Engine, 0x5, "Oss")
+	FOnlineError FromOssError(const FOnlineErrorOss& Result);
+	
+	// if you would like to add a parser that transforms custom platform errors into readable online service types, you can call this in your platform's error handler
+	using TOssPlatformErrorHandler = TUniqueFunction<TOptional<FOnlineError>(const FOnlineError&, const FOnlineErrorOss&)>;
+	void AddOssPlatformErrorHandler(TOssPlatformErrorHandler&& InFunction);
+
+
+} // namespace UE::Online::Errors
+
+
+ONLINESERVICESOSSADAPTER_API bool operator==(const UE::Online::FOnlineError& Left, const FOnlineErrorOss& Right);
+ONLINESERVICESOSSADAPTER_API bool operator==(const FOnlineErrorOss& Left, const UE::Online::FOnlineError& Right);
+ONLINESERVICESOSSADAPTER_API bool operator!=(const UE::Online::FOnlineError& Left, const FOnlineErrorOss& Right);
+ONLINESERVICESOSSADAPTER_API bool operator!=(const FOnlineErrorOss& Left, const UE::Online::FOnlineError& Right);
\ No newline at end of file
