Achievements OSS Adapter implementation

diff --git a/Engine/Plugins/Online/Experimental/OnlineServices/Source/OnlineServicesCommon/Public/Online/OnlineAsyncOp.h b/Engine/Plugins/Online/Experimental/OnlineServices/Source/OnlineServicesCommon/Public/Online/OnlineAsyncOp.h
--- a/Engine/Plugins/Online/Experimental/OnlineServices/Source/OnlineServicesCommon/Public/Online/OnlineAsyncOp.h
+++ b/Engine/Plugins/Online/Experimental/OnlineServices/Source/OnlineServicesCommon/Public/Online/OnlineAsyncOp.h
@@ -956,6 +956,8 @@ protected:
 
 template <typename OpType>
 using TOnlineAsyncOpRef = TSharedRef<TOnlineAsyncOp<OpType>>;
+template <typename OpType>
+using TOnlineAsyncOpPtr = TSharedPtr<TOnlineAsyncOp<OpType>>;
 
 namespace Private {
 
diff --git a/Engine/Plugins/Online/Experimental/OnlineServicesOSSAdapter/Source/Private/Online/AchievementsOSSAdapter.cpp b/Engine/Plugins/Online/Experimental/OnlineServicesOSSAdapter/Source/Private/Online/AchievementsOSSAdapter.cpp
new file mode 100644
--- /dev/null
+++ b/Engine/Plugins/Online/Experimental/OnlineServicesOSSAdapter/Source/Private/Online/AchievementsOSSAdapter.cpp
@@ -0,0 +1,311 @@
+// Copyright Epic Games, Inc. All Rights Reserved.
+
+#include "Online/AchievementsOSSAdapter.h"
+
+#include "Interfaces/OnlineAchievementsInterface.h"
+#include "Interfaces/OnlineExternalUIInterface.h"
+#include "Online/AuthOSSAdapter.h"
+#include "Online/DelegateAdapter.h"
+#include "Online/OnlineErrorDefinitions.h"
+#include "Online/OnlineIdOSSAdapter.h"
+#include "Online/OnlineServicesOSSAdapter.h"
+
+#include "OnlineSubsystem.h"
+
+namespace UE::Online {
+
+void FAchievementsOSSAdapter::PostInitialize()
+{
+	Super::PostInitialize();
+
+	IOnlineSubsystem& SubsystemV1 = static_cast<FOnlineServicesOSSAdapter&>(Services).GetSubsystem();
+
+	AchievementsInterface = SubsystemV1.GetAchievementsInterface();
+	check(AchievementsInterface);
+
+	ExternalUIInterface = SubsystemV1.GetExternalUIInterface();
+
+	AchievementsInterface->AddOnAchievementUnlockedDelegate_Handle(FOnAchievementUnlockedDelegate::CreateLambda([this](const FUniqueNetId& LocalUserIdV1, const FString& AchievementId)
+	{
+		const FOnlineAccountIdHandle LocalUserIdV2 = Services.Get<FAuthOSSAdapter>()->GetAccountIdHandle(LocalUserIdV1.AsShared());
+
+		FAchievementStateUpdated AchievementStateUpdated;
+		AchievementStateUpdated.LocalUserId = LocalUserIdV2;
+		AchievementStateUpdated.AchievementIds = { AchievementId };
+		OnAchievementStateUpdatedEvent.Broadcast(MoveTemp(AchievementStateUpdated));
+	}));
+}
+
+void FAchievementsOSSAdapter::PreShutdown()
+{
+	AchievementsInterface->ClearOnAchievementUnlockedDelegates(this);
+
+	Super::PreShutdown();
+}
+
+TOnlineAsyncOpHandle<FQueryAchievementDefinitions> FAchievementsOSSAdapter::QueryAchievementDefinitions(FQueryAchievementDefinitions::Params&& Params)
+{
+	TOnlineAsyncOpRef<FQueryAchievementDefinitions> Op = GetOp<FQueryAchievementDefinitions>(MoveTemp(Params));
+
+	/**
+	 * Only way to get the achi id's is to call GetCachedAchievements, to get an array of achi objects for this player. Hopefully thats the full set!
+	 * So we first cache the player achi's with QueryAchievements
+	 */
+	Op->Then([this](TOnlineAsyncOp<FQueryAchievementDefinitions>& Op)
+	{
+		TPromise<void> Promise;
+		TFuture<void> Future = Promise.GetFuture();
+
+		const FUniqueNetIdRef LocalUserId = Services.Get<FAuthOSSAdapter>()->GetUniqueNetId(Op.GetParams().LocalUserId);
+		if (!LocalUserId->IsValid())
+		{
+			Op.SetError(Errors::InvalidUser());
+			Promise.SetValue();
+			return Future;
+		}
+
+		if (AchievementDefinitions.IsSet())
+		{
+			Op.SetResult({});
+			Promise.SetValue();
+			return Future;
+		}
+
+		AchievementsInterface->QueryAchievements(*LocalUserId, *MakeDelegateAdapter(this, [this, Promise = MoveTemp(Promise), WeakOp = Op.AsWeak()](const FUniqueNetId& LocalUserId, const bool bWasSuccessful) mutable
+		{
+			if (!bWasSuccessful)
+			{
+				if (TOnlineAsyncOpPtr<FQueryAchievementDefinitions> Op = WeakOp.Pin())
+				{
+					Op->SetError(Errors::Unknown());
+				}
+			}
+			Promise.SetValue();
+		}));
+
+		return Future;
+	})
+	.Then([this](TOnlineAsyncOp<FQueryAchievementDefinitions>& Op)
+	{
+		TPromise<void> Promise;
+		TFuture<void> Future = Promise.GetFuture();
+
+		const FUniqueNetIdRef LocalUserId = Services.Get<FAuthOSSAdapter>()->GetUniqueNetId(Op.GetParams().LocalUserId);
+
+		AchievementsInterface->QueryAchievementDescriptions(*LocalUserId, *MakeDelegateAdapter(this, [this, Promise = MoveTemp(Promise), WeakOp = Op.AsWeak()](const FUniqueNetId& LocalUserId, const bool bWasSuccessful) mutable
+		{
+			if (!bWasSuccessful)
+			{
+				if (TOnlineAsyncOpPtr<FQueryAchievementDefinitions> Op = WeakOp.Pin())
+				{
+					Op->SetError(Errors::Unknown());
+				}
+			}
+			Promise.SetValue();
+		}));
+
+		return Future;
+	})
+	.Then([this](TOnlineAsyncOp<FQueryAchievementDefinitions>& Op)
+	{
+		const FUniqueNetIdRef LocalUserId = Services.Get<FAuthOSSAdapter>()->GetUniqueNetId(Op.GetParams().LocalUserId);
+
+		TArray<FOnlineAchievement> Achievements;
+		EOnlineCachedResult::Type Result = AchievementsInterface->GetCachedAchievements(*LocalUserId, Achievements);
+		if (Result != EOnlineCachedResult::Success)
+		{
+			Op.SetError(Errors::Unknown());
+			return;
+		}
+
+		FAchievementDefinitionMap Definitions;
+		
+		for (FOnlineAchievement& Achievement : Achievements)
+		{
+			FOnlineAchievementDesc AchievementDesc;
+			Result = AchievementsInterface->GetCachedAchievementDescription(Achievement.Id, AchievementDesc);
+			if (Result != EOnlineCachedResult::Success)
+			{
+				Op.SetError(Errors::Unknown());
+				return;
+			}
+
+			FAchievementDefinition Definition;
+			Definition.AchievementId = Achievement.Id;
+			Definition.UnlockedDisplayName = AchievementDesc.Title;
+			Definition.UnlockedDescription = AchievementDesc.UnlockedDesc;
+			Definition.LockedDisplayName = AchievementDesc.Title;
+			Definition.LockedDescription = AchievementDesc.LockedDesc;
+			Definition.bIsHidden = AchievementDesc.bIsHidden;
+			
+			Definitions.Emplace(MoveTemp(Achievement.Id), MoveTemp(Definition));
+		}
+
+		Op.SetResult({});
+	})
+	.Enqueue(GetSerialQueue());
+
+	return Op->GetHandle();
+}
+
+TOnlineResult<FGetAchievementIds> FAchievementsOSSAdapter::GetAchievementIds(FGetAchievementIds::Params&& Params)
+{
+	const FUniqueNetIdRef LocalUserId = Services.Get<FAuthOSSAdapter>()->GetUniqueNetId(Params.LocalUserId);
+	if (!LocalUserId->IsValid())
+	{
+		return TOnlineResult<FGetAchievementIds>(Errors::InvalidUser());
+	}
+
+	if (!AchievementDefinitions.IsSet())
+	{
+		// Call QueryAchievementDefinitions first
+		return TOnlineResult<FGetAchievementIds>(Errors::InvalidState());
+	}
+
+	FGetAchievementIds::Result Result;
+	AchievementDefinitions->GenerateKeyArray(Result.AchievementIds);
+	return TOnlineResult<FGetAchievementIds>(MoveTemp(Result));
+}
+
+TOnlineResult<FGetAchievementDefinition> FAchievementsOSSAdapter::GetAchievementDefinition(FGetAchievementDefinition::Params&& Params)
+{
+	const FUniqueNetIdRef LocalUserId = Services.Get<FAuthOSSAdapter>()->GetUniqueNetId(Params.LocalUserId);
+	if (!LocalUserId->IsValid())
+	{
+		return TOnlineResult<FGetAchievementDefinition>(Errors::InvalidUser());
+	}
+
+	if (!AchievementDefinitions.IsSet())
+	{
+		// Should call QueryAchievementDefinitions first
+		return TOnlineResult<FGetAchievementDefinition>(Errors::InvalidState());
+	}
+
+	const FAchievementDefinition* AchievementDefinition = AchievementDefinitions->Find(Params.AchievementId);
+	if (!AchievementDefinition)
+	{
+		return TOnlineResult<FGetAchievementDefinition>(Errors::NotFound());
+	}
+
+	return TOnlineResult<FGetAchievementDefinition>({ *AchievementDefinition });
+}
+
+TOnlineAsyncOpHandle<FQueryAchievementStates> FAchievementsOSSAdapter::QueryAchievementStates(FQueryAchievementStates::Params&& Params)
+{
+	TOnlineAsyncOpRef<FQueryAchievementStates> Op = GetOp<FQueryAchievementStates>(MoveTemp(Params));
+
+	Op->Then([this](TOnlineAsyncOp<FQueryAchievementStates>& Op)
+	{
+		TPromise<void> Promise;
+		TFuture<void> Future = Promise.GetFuture();
+
+		const FUniqueNetIdRef LocalUserId = Services.Get<FAuthOSSAdapter>()->GetUniqueNetId(Op.GetParams().LocalUserId);
+		if (!LocalUserId->IsValid())
+		{
+			Op.SetError(Errors::InvalidUser());
+			Promise.SetValue();
+			return Future;
+		}
+
+		if (!AchievementDefinitions.IsSet())
+		{
+			// Call QueryAchievementDefinitions first
+			Op.SetError(Errors::InvalidState());
+			Promise.SetValue();
+			return Future;
+		}
+
+		AchievementsInterface->QueryAchievements(*LocalUserId, *MakeDelegateAdapter(this, [this, Promise = MoveTemp(Promise), WeakOp = Op.AsWeak()](const FUniqueNetId& LocalUserId, const bool bWasSuccessful) mutable
+		{
+			if (!bWasSuccessful)
+			{
+				if (TOnlineAsyncOpPtr<FQueryAchievementStates> Op = WeakOp.Pin())
+				{
+					Op->SetError(Errors::Unknown());
+				}
+			}
+			Promise.SetValue();
+		}));
+
+		return Future;
+	})
+	.Then([this](TOnlineAsyncOp<FQueryAchievementStates>& Op)
+	{
+		const FUniqueNetIdRef LocalUserId = Services.Get<FAuthOSSAdapter>()->GetUniqueNetId(Op.GetParams().LocalUserId);
+		TArray<FOnlineAchievement> Achievements;
+		EOnlineCachedResult::Type Result = AchievementsInterface->GetCachedAchievements(*LocalUserId, Achievements);
+		if (Result != EOnlineCachedResult::Success)
+		{
+			Op.SetError(Errors::Unknown());
+			return;
+		}
+
+		FAchievementStateMap& AchievementStates = UserToAchievementStates.Emplace(Op.GetParams().LocalUserId);
+
+		for (FOnlineAchievement& Achievement : Achievements)
+		{
+			FAchievementState& AchievementState = AchievementStates.Emplace(Achievement.Id);
+			AchievementState.AchievementId = Achievement.Id;
+			AchievementState.Progress = Achievement.Progress / 100.f;
+			// AchievementState.UnlockTime left unset because this does not exist on per-user state, but on the common description...
+		}
+
+		Op.SetResult({});
+	})
+	.Enqueue(GetSerialQueue());
+
+	return Op->GetHandle();	
+}
+
+TOnlineResult<FGetAchievementState> FAchievementsOSSAdapter::GetAchievementState(FGetAchievementState::Params&& Params)
+{
+	const FUniqueNetIdRef LocalUserId = Services.Get<FAuthOSSAdapter>()->GetUniqueNetId(Params.LocalUserId);
+	if (!LocalUserId->IsValid())
+	{
+		return TOnlineResult<FGetAchievementState>(Errors::InvalidUser());
+	}
+
+	const FAchievementStateMap* AchievementStates = UserToAchievementStates.Find(Params.LocalUserId);
+	if (!AchievementStates)
+	{
+		// Call QueryAchievementStates first
+		return TOnlineResult<FGetAchievementState>(Errors::InvalidState());
+	}
+
+	const FAchievementState* AchievementState = AchievementStates->Find(Params.AchievementId);
+	if (!AchievementState)
+	{
+		return TOnlineResult<FGetAchievementState>(Errors::NotFound());
+	}
+
+	return TOnlineResult<FGetAchievementState>({ *AchievementState });
+}
+
+TOnlineResult<FDisplayAchievementUI> FAchievementsOSSAdapter::DisplayAchievementUI(FDisplayAchievementUI::Params&& Params)
+{
+	if (!ExternalUIInterface)
+	{
+		return TOnlineResult<FDisplayAchievementUI>(Errors::NotImplemented());
+	}
+
+	const int LocalUserNum = Services.Get<FAuthOSSAdapter>()->GetLocalUserNum(Params.LocalUserId);
+	if (LocalUserNum == INDEX_NONE)
+	{
+		return TOnlineResult<FDisplayAchievementUI>(Errors::InvalidUser());
+	}
+
+	if (Params.AchievementId.IsEmpty())
+	{
+		return TOnlineResult<FDisplayAchievementUI>(Errors::InvalidParams());
+	}
+
+	const bool bSuccess = ExternalUIInterface->ShowAchievementsUI(LocalUserNum);
+	if (!bSuccess)
+	{
+		return TOnlineResult<FDisplayAchievementUI>(Errors::Unknown());
+	}
+
+	return TOnlineResult<FDisplayAchievementUI>(FDisplayAchievementUI::Result());
+}
+
+/* UE::Online */ }
diff --git a/Engine/Plugins/Online/Experimental/OnlineServicesOSSAdapter/Source/Private/Online/OnlineServicesOSSAdapter.cpp b/Engine/Plugins/Online/Experimental/OnlineServicesOSSAdapter/Source/Private/Online/OnlineServicesOSSAdapter.cpp
--- a/Engine/Plugins/Online/Experimental/OnlineServicesOSSAdapter/Source/Private/Online/OnlineServicesOSSAdapter.cpp
+++ b/Engine/Plugins/Online/Experimental/OnlineServicesOSSAdapter/Source/Private/Online/OnlineServicesOSSAdapter.cpp
@@ -2,12 +2,13 @@
 
 #include "Online/OnlineServicesOSSAdapter.h"
 
-#include "Online/OnlineIdOSSAdapter.h"
+#include "Online/AchievementsOSSAdapter.h"
 #include "Online/AuthOSSAdapter.h"
 #include "Online/ConnectivityOSSAdapter.h"
-#include "Online/PrivilegesOSSAdapter.h"
 #include "Online/ExternalUIOSSAdapter.h"
+#include "Online/OnlineIdOSSAdapter.h"
 #include "Online/PresenceOSSAdapter.h"
+#include "Online/PrivilegesOSSAdapter.h"
 #include "Online/StatsOSSAdapter.h"
 
 #include "OnlineSubsystem.h"
@@ -25,10 +26,14 @@ void FOnlineServicesOSSAdapter::RegisterComponents()
 {
 	Components.Register<FAuthOSSAdapter>(*this);
 	Components.Register<FConnectivityOSSAdapter>(*this);
-	Components.Register<FPrivilegesOSSAdapter>(*this);
 	Components.Register<FPresenceOSSAdapter>(*this);
+	Components.Register<FPrivilegesOSSAdapter>(*this);
 	Components.Register<FStatsOSSAdapter>(*this);
 
+	if (Subsystem->GetAchievementsInterface().IsValid())
+	{
+		Components.Register<FAchievementsOSSAdapter>(*this);
+	}
 	if (Subsystem->GetExternalUIInterface().IsValid())
 	{
 		Components.Register<FExternalUIOSSAdapter>(*this);
diff --git a/Engine/Plugins/Online/Experimental/OnlineServicesOSSAdapter/Source/Public/Online/AchievementsOSSAdapter.h b/Engine/Plugins/Online/Experimental/OnlineServicesOSSAdapter/Source/Public/Online/AchievementsOSSAdapter.h
new file mode 100644
--- /dev/null
+++ b/Engine/Plugins/Online/Experimental/OnlineServicesOSSAdapter/Source/Public/Online/AchievementsOSSAdapter.h
@@ -0,0 +1,46 @@
+// Copyright Epic Games, Inc. All Rights Reserved.
+
+#pragma once
+
+#include "Online/AchievementsCommon.h"
+
+#include "OnlineSubsystemTypes.h"
+
+using IOnlineAchievementsPtr = TSharedPtr<class IOnlineAchievements>;
+using IOnlineExternalUIPtr = TSharedPtr<class IOnlineExternalUI>;
+
+namespace UE::Online {
+
+class FAchievementsOSSAdapter : public FAchievementsCommon
+{
+public:
+	using Super = FAchievementsCommon;
+
+	using FAchievementsCommon::FAchievementsCommon;
+
+	// IOnlineComponent
+	virtual void PostInitialize() override;
+	virtual void PreShutdown() override;
+
+	// IAchievements
+	virtual TOnlineAsyncOpHandle<FQueryAchievementDefinitions> QueryAchievementDefinitions(FQueryAchievementDefinitions::Params&& Params) override;
+	virtual TOnlineResult<FGetAchievementIds> GetAchievementIds(FGetAchievementIds::Params&& Params) override;
+	virtual TOnlineResult<FGetAchievementDefinition> GetAchievementDefinition(FGetAchievementDefinition::Params&& Params) override;
+	virtual TOnlineAsyncOpHandle<FQueryAchievementStates> QueryAchievementStates(FQueryAchievementStates::Params&& Params) override;
+	virtual TOnlineResult<FGetAchievementState> GetAchievementState(FGetAchievementState::Params&& Params) override;
+	// Intentionally unimplemented, all v1 achis are stats based
+	//virtual TOnlineAsyncOpHandle<FUnlockAchievements> UnlockAchievements(FUnlockAchievements::Params&& Params) override;
+	virtual TOnlineResult<FDisplayAchievementUI> DisplayAchievementUI(FDisplayAchievementUI::Params&& Params) override;
+
+protected:
+	IOnlineAchievementsPtr AchievementsInterface = nullptr;
+	IOnlineExternalUIPtr ExternalUIInterface = nullptr;
+
+	using FAchievementDefinitionMap = TMap<FString, FAchievementDefinition>;
+	TOptional<FAchievementDefinitionMap> AchievementDefinitions;
+
+	using FAchievementStateMap = TMap<FString, FAchievementState>;
+	TMap<FOnlineAccountIdHandle, FAchievementStateMap> UserToAchievementStates;
+};
+
+/* UE::Online */ }
