OSSv2: Convert OSSv1 online error to OSSv2 online error in Stats adaptor.

diff --git a/Engine/Plugins/Online/Experimental/OnlineServicesOSSAdapter/Source/Private/Online/StatsOSSAdapter.cpp b/Engine/Plugins/Online/Experimental/OnlineServicesOSSAdapter/Source/Private/Online/StatsOSSAdapter.cpp
--- a/Engine/Plugins/Online/Experimental/OnlineServicesOSSAdapter/Source/Private/Online/StatsOSSAdapter.cpp
+++ b/Engine/Plugins/Online/Experimental/OnlineServicesOSSAdapter/Source/Private/Online/StatsOSSAdapter.cpp
@@ -2,6 +2,7 @@
 
 #include "Online/StatsOSSAdapter.h"
 #include "Online/AuthOSSAdapter.h"
+#include "Online/ErrorsOSSAdapter.h"
 
 #include "Online/OnlineServicesOSSAdapter.h"
 #include "Online/OnlineIdOSSAdapter.h"
@@ -97,14 +98,13 @@ TOnlineAsyncOpHandle<FUpdateStats> FStatsOSSAdapter::UpdateStats(FUpdateStats::P
 			return;
 		}
 
-		StatsInterface->UpdateStats(LocalUserId, UpdateUserStatsV1, *MakeDelegateAdapter(this, [WeakOp = Op.AsWeak()](const ::FOnlineError& OnlineError) mutable
+		StatsInterface->UpdateStats(LocalUserId, UpdateUserStatsV1, *MakeDelegateAdapter(this, [WeakOp = Op.AsWeak()](const FOnlineErrorOss& OnlineError) mutable
 		{
 			if (TSharedPtr<TOnlineAsyncOp<FUpdateStats>> PinnedOp = WeakOp.Pin())
 			{
 				if (!OnlineError.WasSuccessful())
 				{
-					// TODO: Convert the error when FOnlineError conversion is ready
-					PinnedOp->SetError(Errors::Unknown());
+					PinnedOp->SetError(Errors::FromOssError(OnlineError));
 					return;
 				}
 
@@ -132,14 +132,13 @@ TOnlineAsyncOpHandle<FQueryStats> FStatsOSSAdapter::QueryStats(FQueryStats::Para
 			return;
 		}
 
-		StatsInterface->QueryStats(LocalUniqueNetId, TargetUniqueNetId, *MakeDelegateAdapter(this, [this, WeakOp = Op.AsWeak()](const ::FOnlineError& OnlineError, const TSharedPtr<const FOnlineStatsUserStats>& QueriedStats) mutable
+		StatsInterface->QueryStats(LocalUniqueNetId, TargetUniqueNetId, *MakeDelegateAdapter(this, [this, WeakOp = Op.AsWeak()](const FOnlineErrorOss& OnlineError, const TSharedPtr<const FOnlineStatsUserStats>& QueriedStats) mutable
 		{
 			if (TSharedPtr<TOnlineAsyncOp<FQueryStats>> PinnedOp = WeakOp.Pin())
 			{
 				if (!OnlineError.WasSuccessful())
 				{
-					// TODO: Convert the error when FOnlineError conversion is ready
-					PinnedOp->SetError(Errors::Unknown()); // TODO: Convert the error?
+					PinnedOp->SetError(Errors::FromOssError(OnlineError));
 					return;
 				}
 
@@ -180,14 +179,13 @@ TOnlineAsyncOpHandle<FBatchQueryStats> FStatsOSSAdapter::BatchQueryStats(FBatchQ
 			NetIds.Add(UniqueNetId);
 		}
 
-		StatsInterface->QueryStats(LocalUserId, NetIds, Op.GetParams().StatNames, *MakeDelegateAdapter(this, [this, WeakOp = Op.AsWeak()](const ::FOnlineError& OnlineError, const TArray<TSharedRef<const FOnlineStatsUserStats>>& UsersStatsResult) mutable
+		StatsInterface->QueryStats(LocalUserId, NetIds, Op.GetParams().StatNames, *MakeDelegateAdapter(this, [this, WeakOp = Op.AsWeak()](const FOnlineErrorOss& OnlineError, const TArray<TSharedRef<const FOnlineStatsUserStats>>& UsersStatsResult) mutable
 		{
 			if (TSharedPtr<TOnlineAsyncOp<FBatchQueryStats>> PinnedOp = WeakOp.Pin())
 			{
 				if (!OnlineError.WasSuccessful())
 				{
-					// TODO: Convert the error when FOnlineError conversion is ready
-					PinnedOp->SetError(Errors::Unknown());
+					PinnedOp->SetError(Errors::FromOssError(OnlineError));
 				}
 
 				FBatchQueryStats::Result Result;
