- Fix TitleFileOSSAdapter not handling requests that failed to start

diff --git a/Engine/Plugins/Online/Experimental/OnlineServicesOSSAdapter/Source/Private/Online/TitleFileOSSAdapter.cpp b/Engine/Plugins/Online/Experimental/OnlineServicesOSSAdapter/Source/Private/Online/TitleFileOSSAdapter.cpp
--- a/Engine/Plugins/Online/Experimental/OnlineServicesOSSAdapter/Source/Private/Online/TitleFileOSSAdapter.cpp
+++ b/Engine/Plugins/Online/Experimental/OnlineServicesOSSAdapter/Source/Private/Online/TitleFileOSSAdapter.cpp
@@ -45,7 +45,11 @@ TOnlineAsyncOpHandle<FTitleFileEnumerateFiles> FTitleFileOSSAdapter::EnumerateFi
 			Promise.SetValue();
 		});
 
-		TitleFileInterface->EnumerateFiles();
+		if (!TitleFileInterface->EnumerateFiles())
+		{
+			Op.SetError(Errors::Unknown());
+			return MakeFulfilledPromise<void>().GetFuture();
+		}
 
 		return Future;
 	})
@@ -95,9 +99,6 @@ TOnlineAsyncOpHandle<FTitleFileReadFile> FTitleFileOSSAdapter::ReadFile(FTitleFi
 
 	Op->Then([this](TOnlineAsyncOp<FTitleFileReadFile>& Op)
 	{
-		TPromise<void> Promise;
-		TFuture<void> Future = Promise.GetFuture();
-
 		const FTitleFileReadFile::Params& Params = Op.GetParams();
 
 		// First just try and read the file contents, if this succeeds we can return immediately.
@@ -106,10 +107,12 @@ TOnlineAsyncOpHandle<FTitleFileReadFile> FTitleFileOSSAdapter::ReadFile(FTitleFi
 		{
 			FTitleFileContentsRef FileContentsRef = MakeShared<FTitleFileContents>(MoveTemp(FileContents));
 			Op.SetResult({FileContentsRef});
-			Promise.SetValue();
-			return Future;
+			return MakeFulfilledPromise<void>().GetFuture();
 		}
 
+		TPromise<void> Promise;
+		TFuture<void> Future = Promise.GetFuture();
+
 		MakeMulticastAdapter(this, TitleFileInterface->OnEnumerateFilesCompleteDelegates,
 			[this, WeakOp = Op.AsWeak(), Promise = MoveTemp(Promise)](bool bSuccess, const FString& ErrorStr) mutable
 		{
@@ -123,7 +126,11 @@ TOnlineAsyncOpHandle<FTitleFileReadFile> FTitleFileOSSAdapter::ReadFile(FTitleFi
 			Promise.SetValue();
 		});
 
-		TitleFileInterface->ReadFile(Params.Filename);
+		if (!TitleFileInterface->ReadFile(Params.Filename))
+		{
+			Op.SetError(Errors::Unknown());
+			return MakeFulfilledPromise<void>().GetFuture();
+		}
 
 		return Future;
 	})
