- UserFile OSS Adapter implementation

diff --git a/Engine/Plugins/Online/Experimental/OnlineServicesOSSAdapter/Source/Private/Online/UserFileOSSAdapter.cpp b/Engine/Plugins/Online/Experimental/OnlineServicesOSSAdapter/Source/Private/Online/UserFileOSSAdapter.cpp
new file mode 100644
--- /dev/null
+++ b/Engine/Plugins/Online/Experimental/OnlineServicesOSSAdapter/Source/Private/Online/UserFileOSSAdapter.cpp
@@ -0,0 +1,329 @@
+// Copyright Epic Games, Inc. All Rights Reserved.
+
+#include "Online/UserFileOSSAdapter.h"
+
+#include "Interfaces/OnlineUserCloudInterface.h"
+#include "Online/AuthOSSAdapter.h"
+#include "Online/DelegateAdapter.h"
+#include "Online/OnlineErrorDefinitions.h"
+#include "Online/OnlineIdOSSAdapter.h"
+#include "Online/OnlineServicesOSSAdapter.h"
+
+#include "OnlineSubsystem.h"
+
+namespace UE::Online {
+
+void FUserFileOSSAdapter::PostInitialize()
+{
+	Super::PostInitialize();
+
+	IOnlineSubsystem& SubsystemV1 = static_cast<FOnlineServicesOSSAdapter&>(Services).GetSubsystem();
+
+	UserCloudInterface = SubsystemV1.GetUserCloudInterface();
+	check(UserCloudInterface);
+}
+
+void FUserFileOSSAdapter::LoadConfig()
+{
+	Super::LoadConfig(Config);
+}
+
+TOnlineAsyncOpHandle<FUserFileEnumerateFiles> FUserFileOSSAdapter::EnumerateFiles(FUserFileEnumerateFiles::Params&& InParams)
+{
+	TOnlineAsyncOpRef<FUserFileEnumerateFiles> Op = GetOp<FUserFileEnumerateFiles>(MoveTemp(InParams));
+
+	Op->Then([this](TOnlineAsyncOp<FUserFileEnumerateFiles>& Op)
+	{
+		const FUniqueNetIdRef LocalUserId = Services.Get<FAuthOSSAdapter>()->GetUniqueNetId(Op.GetParams().LocalUserId);
+		if (!LocalUserId->IsValid())
+		{
+			Op.SetError(Errors::InvalidUser());
+			return MakeFulfilledPromise<void>().GetFuture();
+		}
+
+		TPromise<void> Promise;
+		TFuture<void> Future = Promise.GetFuture();
+
+		MakeMulticastAdapter(this, UserCloudInterface->OnEnumerateUserFilesCompleteDelegates, 
+		[WeakOp = Op.AsWeak(), Promise = MoveTemp(Promise), LocalUserId](bool bSuccess, const FUniqueNetId& UserId) mutable
+		{
+			const bool bThisUser = *LocalUserId == UserId;
+			if(bThisUser)
+			{
+				if (TOnlineAsyncOpPtr<FUserFileEnumerateFiles> Op = WeakOp.Pin())
+				{
+					if (!bSuccess)
+					{
+						Op->SetError(Errors::Unknown());
+					}
+				}
+				Promise.SetValue();
+			}
+			return bThisUser;
+		});
+
+		UserCloudInterface->EnumerateUserFiles(*LocalUserId);
+
+		return Future;
+	})
+	.Then([this](TOnlineAsyncOp<FUserFileEnumerateFiles>& Op)
+	{
+		const FUniqueNetIdRef LocalUserId = Services.Get<FAuthOSSAdapter>()->GetUniqueNetId(Op.GetParams().LocalUserId);
+		if (!LocalUserId->IsValid())
+		{
+			Op.SetError(Errors::InvalidUser());
+			return;
+		}
+
+		TArray<FCloudFileHeader> FileList;
+		UserCloudInterface->GetUserFileList(*LocalUserId, FileList);
+
+		TArray<FString> NewEnumeratedFiles;
+		NewEnumeratedFiles.Reserve(FileList.Num());
+
+		for (const FCloudFileHeader& File : FileList)
+		{
+			NewEnumeratedFiles.Emplace(File.FileName);
+		}
+
+		EnumeratedFiles.Emplace(Op.GetParams().LocalUserId, MoveTemp(NewEnumeratedFiles));
+		Op.SetResult({});
+	})
+	.Enqueue(GetSerialQueue());
+	return Op->GetHandle();
+}
+
+TOnlineResult<FUserFileGetEnumeratedFiles> FUserFileOSSAdapter::GetEnumeratedFiles(FUserFileGetEnumeratedFiles::Params&& Params)
+{
+	const FUniqueNetIdRef LocalUserId = Services.Get<FAuthOSSAdapter>()->GetUniqueNetId(Params.LocalUserId);
+	if (!LocalUserId->IsValid())
+	{
+		return TOnlineResult<FUserFileGetEnumeratedFiles>(Errors::InvalidUser());
+	}
+
+	const TArray<FString>* UserFiles = EnumeratedFiles.Find(Params.LocalUserId);
+	if (!UserFiles)
+	{
+		// Call EnumerateFiles first
+		return TOnlineResult<FUserFileGetEnumeratedFiles>(Errors::InvalidState());
+	}
+
+	return TOnlineResult<FUserFileGetEnumeratedFiles>({ *UserFiles });
+}
+
+TOnlineAsyncOpHandle<FUserFileReadFile> FUserFileOSSAdapter::ReadFile(FUserFileReadFile::Params&& InParams)
+{
+	static const TCHAR* FileContentsKey = TEXT("FileContents");
+
+	TOnlineAsyncOpRef<FUserFileReadFile> Op = GetOp<FUserFileReadFile>(MoveTemp(InParams));
+	const FUserFileReadFile::Params& Params = Op->GetParams();
+
+	if (Params.Filename.IsEmpty())
+	{
+		Op->SetError(Errors::InvalidParams());
+		return Op->GetHandle();
+	}
+
+	Op->Then([this](TOnlineAsyncOp<FUserFileReadFile>& Op)
+	{
+		const FUserFileReadFile::Params& Params = Op.GetParams();
+
+		const FUniqueNetIdRef LocalUserId = Services.Get<FAuthOSSAdapter>()->GetUniqueNetId(Params.LocalUserId);
+		if (!LocalUserId->IsValid())
+		{
+			Op.SetError(Errors::InvalidUser());
+			return MakeFulfilledPromise<void>().GetFuture();
+		}
+
+		// First just try and read the file contents, if this succeeds we can return immediately.
+		TArray<uint8> FileContents;
+		if (UserCloudInterface->GetFileContents(*LocalUserId, Params.Filename, FileContents))
+		{
+			FUserFileContentsRef FileContentsRef = MakeShared<FUserFileContents>(MoveTemp(FileContents));
+			Op.SetResult({FileContentsRef});
+			return MakeFulfilledPromise<void>().GetFuture();
+		}
+
+		TPromise<void> Promise;
+		TFuture<void> Future = Promise.GetFuture();
+
+		MakeMulticastAdapter(this, UserCloudInterface->OnEnumerateUserFilesCompleteDelegates,
+			[WeakOp = Op.AsWeak(), Promise = MoveTemp(Promise), LocalUserId](bool bSuccess, const FUniqueNetId& UserId) mutable
+		{
+			const bool bThisUser = *LocalUserId == UserId;
+			if (bThisUser)
+			{
+				if (TOnlineAsyncOpPtr<FUserFileReadFile> Op = WeakOp.Pin())
+				{
+					if (!bSuccess)
+					{
+						Op->SetError(Errors::Unknown());
+					}
+				}
+				Promise.SetValue();
+			}
+			return bThisUser;
+		});
+
+		if (!UserCloudInterface->ReadUserFile(*LocalUserId, Params.Filename))
+		{
+			Op.SetError(Errors::Unknown());
+			return MakeFulfilledPromise<void>().GetFuture();
+		}
+
+		return Future;
+	})
+	.Then([this](TOnlineAsyncOp<FUserFileReadFile>& Op)
+	{
+		const FUserFileReadFile::Params& Params = Op.GetParams();
+
+		const FUniqueNetIdRef LocalUserId = Services.Get<FAuthOSSAdapter>()->GetUniqueNetId(Params.LocalUserId);
+		if (!LocalUserId->IsValid())
+		{
+			Op.SetError(Errors::InvalidUser());
+			return;
+		}
+
+		TArray<uint8> FileContents;
+		if (UserCloudInterface->GetFileContents(*LocalUserId, Params.Filename, FileContents))
+		{
+			FUserFileContentsRef FileContentsRef = MakeShared<FUserFileContents>(MoveTemp(FileContents));
+			Op.SetResult({ FileContentsRef });
+		}
+		else
+		{
+			Op.SetError(Errors::Unknown());
+		}
+	})
+	.Enqueue(GetSerialQueue());
+	return Op->GetHandle();
+}
+
+TOnlineAsyncOpHandle<FUserFileWriteFile> FUserFileOSSAdapter::WriteFile(FUserFileWriteFile::Params&& InParams)
+{
+	TOnlineAsyncOpRef<FUserFileWriteFile> Op = GetOp<FUserFileWriteFile>(MoveTemp(InParams));
+	const FUserFileWriteFile::Params& Params = Op->GetParams();
+
+	if (Params.Filename.IsEmpty())
+	{
+		Op->SetError(Errors::InvalidParams());
+		return Op->GetHandle();
+	}
+
+	Op->Then([this](TOnlineAsyncOp<FUserFileWriteFile>& Op)
+	{
+		const FUserFileWriteFile::Params& Params = Op.GetParams();
+
+		const FUniqueNetIdRef LocalUserId = Services.Get<FAuthOSSAdapter>()->GetUniqueNetId(Params.LocalUserId);
+		if (!LocalUserId->IsValid())
+		{
+			Op.SetError(Errors::InvalidUser());
+			return MakeFulfilledPromise<void>().GetFuture();
+		}
+
+		TPromise<void> Promise;
+		TFuture<void> Future = Promise.GetFuture();
+
+		MakeMulticastAdapter(this, UserCloudInterface->OnWriteUserFileCompleteDelegates,
+			[WeakOp = Op.AsWeak(), Promise = MoveTemp(Promise), LocalUserId, Filename = Params.Filename](bool bSuccess, const FUniqueNetId& UserId, const FString& InFilename) mutable
+		{
+			const bool bThisUser = *LocalUserId == UserId;
+			const bool bThisFile = Filename != InFilename;
+			const bool bThis = bThisUser && bThisFile;
+			if (bThis)
+			{
+				if (TOnlineAsyncOpPtr<FUserFileWriteFile> Op = WeakOp.Pin())
+				{
+					if (bSuccess)
+					{
+						Op->SetResult({});
+					}
+					else
+					{
+						Op->SetError(Errors::Unknown());
+					}
+				}
+				Promise.SetValue();
+			}
+			return bThis;
+		});
+
+		// WriteUserFile takes a non-const ref, not sure why.
+		FUserFileContents& FileContents = const_cast<FUserFileContents&>(Params.FileContents);
+
+		if (!UserCloudInterface->WriteUserFile(*LocalUserId, Params.Filename, FileContents, Config.bCompressBeforeUpload))
+		{
+			Op.SetError(Errors::Unknown());
+			return MakeFulfilledPromise<void>().GetFuture();
+		}
+
+		return Future;
+	})
+	.Enqueue(GetSerialQueue());
+	return Op->GetHandle();
+}
+
+TOnlineAsyncOpHandle<FUserFileDeleteFile> FUserFileOSSAdapter::DeleteFile(FUserFileDeleteFile::Params&& InParams)
+{
+	TOnlineAsyncOpRef<FUserFileDeleteFile> Op = GetOp<FUserFileDeleteFile>(MoveTemp(InParams));
+	const FUserFileDeleteFile::Params& Params = Op->GetParams();
+
+	if (Params.Filename.IsEmpty())
+	{
+		Op->SetError(Errors::InvalidParams());
+		return Op->GetHandle();
+	}
+
+	Op->Then([this](TOnlineAsyncOp<FUserFileDeleteFile>& Op)
+	{
+		const FUserFileDeleteFile::Params& Params = Op.GetParams();
+
+		const FUniqueNetIdRef LocalUserId = Services.Get<FAuthOSSAdapter>()->GetUniqueNetId(Params.LocalUserId);
+		if (!LocalUserId->IsValid())
+		{
+			Op.SetError(Errors::InvalidUser());
+			return MakeFulfilledPromise<void>().GetFuture();
+		}
+
+		TPromise<void> Promise;
+		TFuture<void> Future = Promise.GetFuture();
+
+		MakeMulticastAdapter(this, UserCloudInterface->OnDeleteUserFileCompleteDelegates,
+			[WeakOp = Op.AsWeak(), Promise = MoveTemp(Promise), LocalUserId, Filename = Params.Filename](bool bSuccess, const FUniqueNetId& UserId, const FString& InFilename) mutable
+		{
+			const bool bThisUser = *LocalUserId == UserId;
+			const bool bThisFile = Filename != InFilename;
+			const bool bThis = bThisUser && bThisFile;
+			if (bThis)
+			{
+				if (TOnlineAsyncOpPtr<FUserFileDeleteFile> Op = WeakOp.Pin())
+				{
+					if (bSuccess)
+					{
+						Op->SetResult({});
+					}
+					else
+					{
+						Op->SetError(Errors::Unknown());
+					}
+				}
+				Promise.SetValue();
+			}
+			return bThis;
+		});
+
+		const bool bShouldCloudDelete = true;
+		const bool bShouldLocallyDelete = true;
+		if (!UserCloudInterface->DeleteUserFile(*LocalUserId, Params.Filename, bShouldCloudDelete, bShouldLocallyDelete))
+		{
+			Op.SetError(Errors::Unknown());
+			return MakeFulfilledPromise<void>().GetFuture();
+		}
+
+		return Future;
+	})
+	.Enqueue(GetSerialQueue());
+	return Op->GetHandle();
+}
+
+/* UE::Online */ }
diff --git a/Engine/Plugins/Online/Experimental/OnlineServicesOSSAdapter/Source/Public/Online/MulticastAdapter.h b/Engine/Plugins/Online/Experimental/OnlineServicesOSSAdapter/Source/Public/Online/MulticastAdapter.h
--- a/Engine/Plugins/Online/Experimental/OnlineServicesOSSAdapter/Source/Public/Online/MulticastAdapter.h
+++ b/Engine/Plugins/Online/Experimental/OnlineServicesOSSAdapter/Source/Public/Online/MulticastAdapter.h
@@ -2,6 +2,8 @@
 #pragma once
 #include "OnlineSubsystemTypes.h"
 
+#include <type_traits>
+
 /*
 	BASIC USAGE
 		The multicast adapter maintains the same functionality as a delegate adapter
@@ -10,15 +12,17 @@
 			->
 			MakeMulticastAdapter(this, Identity->OnLoginComplete, [this](int32 UserNum){...});
 
-		This type of adapter will unbind itself immediately after a single execution. 
+		When this adapter unbinds itself depends on the return type of the lambda bound to it:
+		- void - the adapter will unbind itself after a single execution
+		- bool - the adapter will unbind itself when the lambda returns true, and will stay bound when it returns false.
 
 		The lifetime of the adapter is a shared ptr attached to the input delegate and will live as long as that delegate is bound.
 */
 namespace UE::Online {
 
-template<typename ComponentType, typename DelegateType, typename... LambdaArgs>
+template<typename ComponentType, typename DelegateType, typename LambdaRet, typename... LambdaArgs>
 class TMulticastDelegateAdapter
-	: public TSharedFromThis<TMulticastDelegateAdapter<ComponentType, DelegateType, LambdaArgs...>>
+	: public TSharedFromThis<TMulticastDelegateAdapter<ComponentType, DelegateType, LambdaRet, LambdaArgs...>>
 {
 
 public:
@@ -28,25 +32,36 @@ public:
 	{
 	}
 
-	void SetupDelegate(TUniqueFunction<void(LambdaArgs...)>&& InCallback)
+	void SetupDelegate(TUniqueFunction<LambdaRet(LambdaArgs...)>&& InCallback)
 	{
-		TSharedPtr<TMulticastDelegateAdapter<ComponentType, DelegateType, LambdaArgs...>> SharedAnchor = this->AsShared();
+		TSharedPtr<TMulticastDelegateAdapter<ComponentType, DelegateType, LambdaRet, LambdaArgs...>> SharedAnchor = this->AsShared();
 
 		// callback exists outside of the delegate's bound lambda since that lambda is not move-only
 		Callback = [this, InCallback = MoveTemp(InCallback)](LambdaArgs... Args)
 		{
-			InCallback(Forward<LambdaArgs>(Args)...);
+			return InCallback(Forward<LambdaArgs>(Args)...);
 		};
 
 		Handle = Delegate.AddLambda([this, SharedAnchor](LambdaArgs... Args) mutable 
 		{
+			bool bUnbind = true;
 			if(Parent.IsValid())
 			{
-				Callback(Forward<LambdaArgs>(Args)...);
+				if constexpr (std::is_void_v<LambdaRet>)
+				{
+					Callback(Forward<LambdaArgs>(Args)...);
+				}
+				else
+				{
+					bUnbind = Callback(Forward<LambdaArgs>(Args)...);
+				}
 			}
 
-			SharedAnchor.Reset();
-			Delegate.Remove(Handle);
+			if (bUnbind)
+			{
+				SharedAnchor.Reset();
+				Delegate.Remove(Handle);
+			}
 		});
 	}
 
@@ -55,14 +70,14 @@ public:
 		return Handle;
 	}
 
-	TWeakPtr<TMulticastDelegateAdapter<ComponentType, DelegateType, LambdaArgs...>> AsWeak()
+	TWeakPtr<TMulticastDelegateAdapter<ComponentType, DelegateType, LambdaRet, LambdaArgs...>> AsWeak()
 	{
-		return TWeakPtr<TMulticastDelegateAdapter<ComponentType, DelegateType, LambdaArgs...>>(this->AsShared());
+		return TWeakPtr<TMulticastDelegateAdapter<ComponentType, DelegateType, LambdaRet, LambdaArgs...>>(this->AsShared());
 	}
 
 private:
 	TWeakPtr<ComponentType> Parent;
-	TUniqueFunction<void(LambdaArgs...)> Callback;
+	TUniqueFunction<LambdaRet(LambdaArgs...)> Callback;
 	DelegateType& Delegate;
 	FDelegateHandle Handle;
 };
@@ -88,9 +103,9 @@ namespace Private
 		}
 
 		template<typename ComponentType, typename DelegateType>
-		TSharedPtr<TMulticastDelegateAdapter<ComponentType, DelegateType, ParamTypes...>> Construct(TSharedRef<ComponentType> Interface, DelegateType& InDelegate, CallableObject&& InCallback)
+		TSharedPtr<TMulticastDelegateAdapter<ComponentType, DelegateType, TResultType, ParamTypes...>> Construct(TSharedRef<ComponentType> Interface, DelegateType& InDelegate, CallableObject&& InCallback)
 		{
-			TSharedPtr<TMulticastDelegateAdapter<ComponentType, DelegateType, ParamTypes...>> Adapter = MakeShared<TMulticastDelegateAdapter<ComponentType, DelegateType, ParamTypes...>>(Interface, InDelegate);
+			TSharedPtr<TMulticastDelegateAdapter<ComponentType, DelegateType, TResultType, ParamTypes...>> Adapter = MakeShared<TMulticastDelegateAdapter<ComponentType, DelegateType, TResultType, ParamTypes...>>(Interface, InDelegate);
 			Adapter->SetupDelegate(GetUniqueFn(MoveTemp(InCallback)));
 			return Adapter;
 		}
diff --git a/Engine/Plugins/Online/Experimental/OnlineServicesOSSAdapter/Source/Public/Online/UserFileOSSAdapter.h b/Engine/Plugins/Online/Experimental/OnlineServicesOSSAdapter/Source/Public/Online/UserFileOSSAdapter.h
new file mode 100644
--- /dev/null
+++ b/Engine/Plugins/Online/Experimental/OnlineServicesOSSAdapter/Source/Public/Online/UserFileOSSAdapter.h
@@ -0,0 +1,50 @@
+// Copyright Epic Games, Inc. All Rights Reserved.
+
+#pragma once
+
+#include "Online/UserFileCommon.h"
+
+#include "OnlineSubsystemTypes.h"
+
+using IOnlineUserCloudPtr = TSharedPtr<class IOnlineUserCloud>;
+
+namespace UE::Online {
+
+struct FUserFileOSSAdapterConfig
+{
+	bool bCompressBeforeUpload = false;
+};
+
+namespace Meta {
+
+BEGIN_ONLINE_STRUCT_META(FUserFileOSSAdapterConfig)
+	ONLINE_STRUCT_FIELD(FUserFileOSSAdapterConfig, bCompressBeforeUpload)
+END_ONLINE_STRUCT_META()
+/* Meta */ }
+
+class FUserFileOSSAdapter : public FUserFileCommon
+{
+public:
+	using Super = FUserFileCommon;
+
+	using FUserFileCommon::FUserFileCommon;
+
+	// IOnlineComponent
+	virtual void PostInitialize() override;
+	virtual void LoadConfig() override;
+
+	// IUserFile
+	virtual TOnlineAsyncOpHandle<FUserFileEnumerateFiles> EnumerateFiles(FUserFileEnumerateFiles::Params&& Params) override;
+	virtual TOnlineResult<FUserFileGetEnumeratedFiles> GetEnumeratedFiles(FUserFileGetEnumeratedFiles::Params&& Params) override;
+	virtual TOnlineAsyncOpHandle<FUserFileReadFile> ReadFile(FUserFileReadFile::Params&& Params) override;
+	virtual TOnlineAsyncOpHandle<FUserFileWriteFile> WriteFile(FUserFileWriteFile::Params&& Params) override;
+	//virtual TOnlineAsyncOpHandle<FUserFileCopyFile> CopyFile(FUserFileCopyFile::Params&& Params) override; // Unimplemented
+	virtual TOnlineAsyncOpHandle<FUserFileDeleteFile> DeleteFile(FUserFileDeleteFile::Params&& Params) override;
+
+protected:
+	FUserFileOSSAdapterConfig Config;
+	IOnlineUserCloudPtr UserCloudInterface = nullptr;
+	TMap<FOnlineAccountIdHandle, TArray<FString>> EnumeratedFiles;
+};
+
+/* UE::Online */ }
