- Fix some LoadConfig implementations not calling Super::LoadConfig();

diff --git a/Engine/Plugins/Online/Experimental/OnlineServicesEOSGS/Source/Private/Online/TitleFileEOSGS.cpp b/Engine/Plugins/Online/Experimental/OnlineServicesEOSGS/Source/Private/Online/TitleFileEOSGS.cpp
--- a/Engine/Plugins/Online/Experimental/OnlineServicesEOSGS/Source/Private/Online/TitleFileEOSGS.cpp
+++ b/Engine/Plugins/Online/Experimental/OnlineServicesEOSGS/Source/Private/Online/TitleFileEOSGS.cpp
@@ -29,6 +29,7 @@ void FTitleFileEOSGS::Initialize()
 
 void FTitleFileEOSGS::LoadConfig()
 {
+	Super::LoadConfig();
 	Super::LoadConfig(Config);
 }
 
diff --git a/Engine/Plugins/Online/Experimental/OnlineServicesEOSGS/Source/Private/Online/UserFileEOSGS.cpp b/Engine/Plugins/Online/Experimental/OnlineServicesEOSGS/Source/Private/Online/UserFileEOSGS.cpp
--- a/Engine/Plugins/Online/Experimental/OnlineServicesEOSGS/Source/Private/Online/UserFileEOSGS.cpp
+++ b/Engine/Plugins/Online/Experimental/OnlineServicesEOSGS/Source/Private/Online/UserFileEOSGS.cpp
@@ -32,6 +32,7 @@ void FUserFileEOSGS::Initialize()
 
 void FUserFileEOSGS::LoadConfig()
 {
+	Super::LoadConfig();
 	Super::LoadConfig(Config);
 }
 
diff --git a/Engine/Plugins/Online/Experimental/OnlineServicesOSSAdapter/Source/Private/Online/UserFileOSSAdapter.cpp b/Engine/Plugins/Online/Experimental/OnlineServicesOSSAdapter/Source/Private/Online/UserFileOSSAdapter.cpp
--- a/Engine/Plugins/Online/Experimental/OnlineServicesOSSAdapter/Source/Private/Online/UserFileOSSAdapter.cpp
+++ b/Engine/Plugins/Online/Experimental/OnlineServicesOSSAdapter/Source/Private/Online/UserFileOSSAdapter.cpp
@@ -25,6 +25,7 @@ void FUserFileOSSAdapter::PostInitialize()
 
 void FUserFileOSSAdapter::LoadConfig()
 {
+	Super::LoadConfig();
 	Super::LoadConfig(Config);
 }
 
