- Fix ReadUserFile hanging indefinitely due to binding accidentally to the Enumerate callback
- Fix incorrect operator used to compare filenames.

diff --git a/Engine/Plugins/Online/Experimental/OnlineServicesOSSAdapter/Source/Private/Online/UserFileOSSAdapter.cpp b/Engine/Plugins/Online/Experimental/OnlineServicesOSSAdapter/Source/Private/Online/UserFileOSSAdapter.cpp
--- a/Engine/Plugins/Online/Experimental/OnlineServicesOSSAdapter/Source/Private/Online/UserFileOSSAdapter.cpp
+++ b/Engine/Plugins/Online/Experimental/OnlineServicesOSSAdapter/Source/Private/Online/UserFileOSSAdapter.cpp
@@ -148,11 +148,13 @@ TOnlineAsyncOpHandle<FUserFileReadFile> FUserFileOSSAdapter::ReadFile(FUserFileR
 		TPromise<void> Promise;
 		TFuture<void> Future = Promise.GetFuture();
 
-		MakeMulticastAdapter(this, UserCloudInterface->OnEnumerateUserFilesCompleteDelegates,
-			[WeakOp = Op.AsWeak(), Promise = MoveTemp(Promise), LocalUserId](bool bSuccess, const FUniqueNetId& UserId) mutable
+		MakeMulticastAdapter(this, UserCloudInterface->OnReadUserFileCompleteDelegates,
+			[WeakOp = Op.AsWeak(), Promise = MoveTemp(Promise), LocalUserId, Filename = Params.Filename](bool bSuccess, const FUniqueNetId& UserId, const FString& InFilename) mutable
 		{
 			const bool bThisUser = *LocalUserId == UserId;
-			if (bThisUser)
+			const bool bThisFile = Filename == InFilename;
+			const bool bThis = bThisUser && bThisFile;
+			if (bThis)
 			{
 				if (TOnlineAsyncOpPtr<FUserFileReadFile> Op = WeakOp.Pin())
 				{
@@ -163,7 +165,7 @@ TOnlineAsyncOpHandle<FUserFileReadFile> FUserFileOSSAdapter::ReadFile(FUserFileR
 				}
 				Promise.SetValue();
 			}
-			return bThisUser;
+			return bThis;
 		});
 
 		if (!UserCloudInterface->ReadUserFile(*LocalUserId, Params.Filename))
@@ -229,7 +231,7 @@ TOnlineAsyncOpHandle<FUserFileWriteFile> FUserFileOSSAdapter::WriteFile(FUserFil
 			[WeakOp = Op.AsWeak(), Promise = MoveTemp(Promise), LocalUserId, Filename = Params.Filename](bool bSuccess, const FUniqueNetId& UserId, const FString& InFilename) mutable
 		{
 			const bool bThisUser = *LocalUserId == UserId;
-			const bool bThisFile = Filename != InFilename;
+			const bool bThisFile = Filename == InFilename;
 			const bool bThis = bThisUser && bThisFile;
 			if (bThis)
 			{
@@ -293,7 +295,7 @@ TOnlineAsyncOpHandle<FUserFileDeleteFile> FUserFileOSSAdapter::DeleteFile(FUserF
 			[WeakOp = Op.AsWeak(), Promise = MoveTemp(Promise), LocalUserId, Filename = Params.Filename](bool bSuccess, const FUniqueNetId& UserId, const FString& InFilename) mutable
 		{
 			const bool bThisUser = *LocalUserId == UserId;
-			const bool bThisFile = Filename != InFilename;
+			const bool bThisFile = Filename == InFilename;
 			const bool bThis = bThisUser && bThisFile;
 			if (bThis)
 			{
