Auth OSS adapter

diff --git a/Engine/Plugins/Online/Experimental/OnlineServicesOSSAdapter/Source/Private/Online/AuthOSSAdapter.cpp b/Engine/Plugins/Online/Experimental/OnlineServicesOSSAdapter/Source/Private/Online/AuthOSSAdapter.cpp
--- a/Engine/Plugins/Online/Experimental/OnlineServicesOSSAdapter/Source/Private/Online/AuthOSSAdapter.cpp
+++ b/Engine/Plugins/Online/Experimental/OnlineServicesOSSAdapter/Source/Private/Online/AuthOSSAdapter.cpp
@@ -2,6 +2,7 @@
 
 #include "Online/AuthOSSAdapter.h"
 
+#include "Online/AuthErrors.h"
 #include "Online/OnlineServicesOSSAdapter.h"
 #include "Online/OnlineIdOSSAdapter.h"
 #include "Online/OnlineErrorDefinitions.h"
@@ -13,6 +14,114 @@
 
 namespace UE::Online {
 
+namespace {
+
+static const FString InitialLocalUserNumKeyName = TEXT("InitialLocalUserNum");
+static const FString AccountInfoKeyName = TEXT("AccountInfoOSSAdapter");
+static const FString ExternalAuthTokenTranslationTraitsKeyName = TEXT("ExternalAuthTokenTranslationTraits");
+static const FString ExternalServerAuthTicketTranslationTraitsKeyName = TEXT("ExternalServerAuthTicketTranslationTraits");
+
+ELoginStatus TranslateLoginStatus(::ELoginStatus::Type Status)
+{
+	switch (Status)
+	{
+	case ::ELoginStatus::NotLoggedIn:
+		return ELoginStatus::NotLoggedIn;
+	case ::ELoginStatus::UsingLocalProfile:
+		return ELoginStatus::UsingLocalProfile;
+	case ::ELoginStatus::LoggedIn:
+		return ELoginStatus::LoggedIn;
+	default:
+		checkNoEntry();
+		return ELoginStatus::NotLoggedIn;
+	}
+}
+
+/* anonymous */ }
+
+enum class EExternalAuthTokenTranslationFlags : uint8
+{
+	None = 0,
+	TokenString = 1 << 0,
+	TokenBinary = 1 << 1,
+	AuthToken = 1 << 2,
+};
+ENUM_CLASS_FLAGS(EExternalAuthTokenTranslationFlags);
+
+struct FExternalAuthTokenTranslationTraits
+{
+	FName ExternalLoginType;
+	EExternalAuthTokenMethod Method;
+	EExternalAuthTokenTranslationFlags Flags;
+};
+
+const FExternalAuthTokenTranslationTraits* GetExternalAuthTokenTranslationTraits(FName SubsystemType, EExternalAuthTokenMethod Method)
+{
+	static const TMap<FName, TArray<FExternalAuthTokenTranslationTraits>> SupportedExternalAuthTranslatorTraits = {
+		{ TEXT("GDK"), { { ExternalLoginType::XblXstsToken, EExternalAuthTokenMethod::Primary, EExternalAuthTokenTranslationFlags::TokenString } } },
+		{ PS4_SUBSYSTEM, { { ExternalLoginType::PsnIdToken, EExternalAuthTokenMethod::Primary, EExternalAuthTokenTranslationFlags::TokenString } } },
+		{ TEXT("PS5"), { { ExternalLoginType::PsnIdToken, EExternalAuthTokenMethod::Primary, EExternalAuthTokenTranslationFlags::TokenString } } },
+		{ SWITCH_SUBSYSTEM, { { ExternalLoginType::NintendoNsaIdToken, EExternalAuthTokenMethod::Primary, EExternalAuthTokenTranslationFlags::AuthToken },
+							  { ExternalLoginType::NintendoIdToken, EExternalAuthTokenMethod::Secondary, EExternalAuthTokenTranslationFlags::TokenString } } },
+		{ STEAM_SUBSYSTEM, { { ExternalLoginType::SteamAppTicket, EExternalAuthTokenMethod::Primary, EExternalAuthTokenTranslationFlags::TokenBinary } } },
+	};
+
+	if (const TArray<FExternalAuthTokenTranslationTraits>* TraitsArray = SupportedExternalAuthTranslatorTraits.Find(SubsystemType))
+	{
+		return TraitsArray->FindByPredicate([Method](const FExternalAuthTokenTranslationTraits& Value) { return Value.Method == Method; });
+	}
+
+	return nullptr;
+}
+
+struct FExternalServerAuthTicketTranslationTraits
+{
+	FName ExternalTicketType;
+};
+
+TDefaultErrorResultInternal<TSharedRef<FExternalServerAuthTicketTranslationTraits>> GetExternalServerAuthTicketTranslationTraits(FName SubsystemType)
+{
+	static const TMap<FName, TSharedRef<FExternalServerAuthTicketTranslationTraits>> SupportedExternalServerAuthTicketTranslationTraits = {
+		{ TEXT("GDK"), MakeShared<FExternalServerAuthTicketTranslationTraits>(FExternalServerAuthTicketTranslationTraits{ ExternalServerAuthTicketType::XblXstsToken }) },
+		{ PS4_SUBSYSTEM, MakeShared<FExternalServerAuthTicketTranslationTraits>(FExternalServerAuthTicketTranslationTraits{ ExternalServerAuthTicketType::PsnAuthCode }) },
+		{ TEXT("PS5"), MakeShared<FExternalServerAuthTicketTranslationTraits>(FExternalServerAuthTicketTranslationTraits{ ExternalServerAuthTicketType::PsnAuthCode }) },
+	};
+
+	if (const TSharedRef<FExternalServerAuthTicketTranslationTraits>* Traits = SupportedExternalServerAuthTicketTranslationTraits.Find(SubsystemType))
+	{
+		return TDefaultErrorResultInternal<TSharedRef<FExternalServerAuthTicketTranslationTraits>>(*Traits);
+	}
+
+	return TDefaultErrorResultInternal<TSharedRef<FExternalServerAuthTicketTranslationTraits>>(Errors::NotImplemented());
+}
+
+TSharedPtr<FAccountInfoOSSAdapter> FAccountInfoRegistryOSSAdapter::Find(FPlatformUserId PlatformUserId) const
+{
+	return StaticCastSharedPtr<FAccountInfoOSSAdapter>(Super::Find(PlatformUserId));
+} 
+
+TSharedPtr<FAccountInfoOSSAdapter> FAccountInfoRegistryOSSAdapter::Find(FOnlineAccountIdHandle AccountIdHandle) const
+{
+	return StaticCastSharedPtr<FAccountInfoOSSAdapter>(Super::Find(AccountIdHandle));
+}
+
+void FAccountInfoRegistryOSSAdapter::Register(const TSharedRef<FAccountInfoOSSAdapter>& AccountInfoNULL)
+{
+	DoRegister(AccountInfoNULL);
+}
+
+void FAccountInfoRegistryOSSAdapter::Unregister(FOnlineAccountIdHandle AccountId)
+{
+	if (TSharedPtr<FAccountInfoOSSAdapter> AccountInfoNULL = Find(AccountId))
+	{
+		DoUnregister(AccountInfoNULL.ToSharedRef());
+	}
+	else
+	{
+		UE_LOG(LogOnlineServices, Warning, TEXT("[FAccountInfoRegistryOSSAdapter::Unregister] Failed to find account [%s]."), *ToLogString(AccountId));
+	}
+}
+
 void FAuthOSSAdapter::PostInitialize()
 {
 	Super::PostInitialize();
@@ -22,41 +131,14 @@ void FAuthOSSAdapter::PostInitialize()
 		for (int LocalPlayerNum = 0; LocalPlayerNum < MAX_LOCAL_PLAYERS; ++LocalPlayerNum)
 		{
 			OnLoginStatusChangedHandle[LocalPlayerNum] = Identity->OnLoginStatusChangedDelegates[LocalPlayerNum].AddLambda(
-				[WeakThis = TWeakPtr<IAuth>(AsShared()), this](int32 LocalUserNum, ::ELoginStatus::Type OldStatus, ::ELoginStatus::Type NewStatus, const FUniqueNetId& NewId)
+				[WeakThis = TWeakPtr<FAuthOSSAdapter>(StaticCastSharedRef<FAuthOSSAdapter>(AsShared()))](int32 LocalUserNum, ::ELoginStatus::Type OldStatus, ::ELoginStatus::Type NewStatus, const FUniqueNetId& NetId)
 				{
-					TSharedPtr<IAuth> PinnedThis = WeakThis.Pin();
-					if (PinnedThis.IsValid())
+					if (TSharedPtr<FAuthOSSAdapter> PinnedThis = WeakThis.Pin())
 					{
-						FLoginStatusChanged LoginStatusChanged;
-						if (NewId.IsValid())
-						{
-							LoginStatusChanged.LocalUserId = static_cast<FOnlineServicesOSSAdapter&>(Services).GetAccountIdRegistry().FindOrAddHandle(NewId.AsShared());
-						}
-						switch (OldStatus)
-						{
-							case ::ELoginStatus::NotLoggedIn:
-								LoginStatusChanged.PreviousStatus = ELoginStatus::NotLoggedIn;
-								break;
-							case ::ELoginStatus::UsingLocalProfile:
-								LoginStatusChanged.PreviousStatus = ELoginStatus::UsingLocalProfile;
-								break;
-							case ::ELoginStatus::LoggedIn:
-								LoginStatusChanged.PreviousStatus = ELoginStatus::LoggedIn;
-								break;
-						}
-						switch (NewStatus)
-						{
-							case ::ELoginStatus::NotLoggedIn:
-								LoginStatusChanged.CurrentStatus = ELoginStatus::NotLoggedIn;
-								break;
-							case ::ELoginStatus::UsingLocalProfile:
-								LoginStatusChanged.CurrentStatus = ELoginStatus::UsingLocalProfile;
-								break;
-							case ::ELoginStatus::LoggedIn:
-								LoginStatusChanged.CurrentStatus = ELoginStatus::LoggedIn;
-								break;
-						}
-						OnLoginStatusChangedEvent.Broadcast(LoginStatusChanged);
+						PinnedThis->HandleLoginStatusChangedImplOp(FHandleLoginStatusChangedImpl::Params{
+							FPlatformMisc::GetPlatformUserForUserIndex(LocalUserNum),
+							static_cast<FOnlineServicesOSSAdapter&>(PinnedThis->Services).GetAccountIdRegistry().FindOrAddHandle(NetId.AsShared()),
+							TranslateLoginStatus(NewStatus) });
 					}
 				});
 		}
@@ -80,58 +162,138 @@ TOnlineAsyncOpHandle<FAuthLogin> FAuthOSSAdapter::Login(FAuthLogin::Params&& Par
 {
 	TSharedRef<TOnlineAsyncOp<FAuthLogin>> Op = GetOp<FAuthLogin>(MoveTemp(Params));
 
-	if (!Op->IsReady())
+	// Step 1: Setup operation data.
+	Op->Then([this](TOnlineAsyncOp<FAuthLogin>& InAsyncOp)
 	{
-		int32 LocalUserIndex = FPlatformMisc::GetUserIndexForPlatformUser(Op->GetParams().PlatformUserId);
+		const FAuthLogin::Params& Params = InAsyncOp.GetParams();
 
-		if (LocalUserIndex < 0 || LocalUserIndex >= MAX_LOCAL_PLAYERS)
+		const int32 InitialLocalUserNum = GetIdentityInterface()->GetLocalUserNumFromPlatformUserId(Params.PlatformUserId);
+		if (InitialLocalUserNum == INDEX_NONE)
 		{
-			Op->SetError(Errors::InvalidParams());
-			return Op->GetHandle();
+			InAsyncOp.SetError(Errors::InvalidParams());
+			return;
 		}
 
-		Op->Then([this, LocalUserIndex](TOnlineAsyncOp<FAuthLogin>& Op, TPromise<TOnlineResult<FAuthLogin>>&& Result)
-			{
-				FOnlineAccountCredentials Credentials;
-				Credentials.Type = Op.GetParams().CredentialsType;
-				Credentials.Id = Op.GetParams().CredentialsId;
-				Credentials.Token = Op.GetParams().CredentialsToken.Get<FString>(); // TODO: handle binary token
+		// Set initial user num on operation. Depending on the implementation the local user num
+		// which completes login may be different from the one which started it.
+		InAsyncOp.Data.Set<int32>(InitialLocalUserNumKeyName, InitialLocalUserNum);
+
+		if (AccountInfoRegistryOSSAdapter.Find(Params.PlatformUserId))
+		{
+			InAsyncOp.SetError(Errors::Auth::AlreadyLoggedIn());
+			return;
+		}
+
+		// Todo: Handle platforms that allow calling login again with a different login type.
+
+		TSharedPtr<FAccountInfoOSSAdapter> AccountInfoOSSAdapter = MakeShared<FAccountInfoOSSAdapter>();
+		AccountInfoOSSAdapter->LoginStatus = ELoginStatus::NotLoggedIn;
+
+		// Set user auth data on operation.
+		InAsyncOp.Data.Set<TSharedRef<FAccountInfoOSSAdapter>>(AccountInfoKeyName, AccountInfoOSSAdapter.ToSharedRef());
+	})
+	// Step 2: Login to OSSv1 identity interface.
+	.Then([this](TOnlineAsyncOp<FAuthLogin>& InAsyncOp)
+	{
+		const TSharedRef<FAccountInfoOSSAdapter>& AccountInfoOSSAdapter = GetOpDataChecked<TSharedRef<FAccountInfoOSSAdapter>>(InAsyncOp, AccountInfoKeyName);
+		const int InitialLocalUserNum = GetOpDataChecked<int32>(InAsyncOp, InitialLocalUserNumKeyName);
 
-				MakeMulticastAdapter(this, GetIdentityInterface()->OnLoginCompleteDelegates[LocalUserIndex], 
-				[this, ResultPromise = MoveTemp(Result), PlatformUserId = Op.GetParams().PlatformUserId, LocalUserIndex](int32 LocalUserNum, bool bWasSuccessful, const FUniqueNetId& UserId, const FString& Error) mutable
+		TPromise<void> Promise;
+		TFuture<void> Future = Promise.GetFuture();
+		MakeMulticastAdapter(this, GetIdentityInterface()->OnLoginCompleteDelegates[InitialLocalUserNum],
+		[this, AccountInfoOSSAdapter, WeakOp = InAsyncOp.AsWeak(), Promise = MoveTemp(Promise)](int32 LocalUserNum, bool bWasSuccessful, const FUniqueNetId& UserId, const FString& Error) mutable
+		{
+			if (TSharedPtr<TOnlineAsyncOp<FAuthLogin>> Op = WeakOp.Pin())
+			{
+				if (bWasSuccessful)
 				{
-					if (bWasSuccessful)
+					// Sanity check UserId
+					if (!UserId.IsValid())
 					{
-						FOnlineAccountIdHandle Handle = static_cast<FOnlineServicesOSSAdapter&>(Services).GetAccountIdRegistry().FindOrAddHandle(UserId.AsShared());
-						FAuthLogin::Result Result = { MakeShared<FAccountInfo>() };
-						Result.AccountInfo->PlatformUserId = PlatformUserId;
-						Result.AccountInfo->UserId = Handle;
-						Result.AccountInfo->LoginStatus = ELoginStatus::LoggedIn;
-						ResultPromise.EmplaceValue(MoveTemp(Result));
+						UE_LOG(LogOnlineServices, Warning, TEXT("[FAuthOSSAdapter::Login][%s] Failure: Userid was invalid."), *GetSubsystem().GetSubsystemName().ToString());
+						Op->SetError(Errors::Unknown());
 					}
 					else
 					{
-						FOnlineError V2Error = Errors::Unknown(); // TODO: V1 to V2 error conversion/error from string conversion
-						ResultPromise.EmplaceValue(MoveTemp(V2Error));
+						AccountInfoOSSAdapter->AccountId = static_cast<FOnlineServicesOSSAdapter&>(Services).GetAccountIdRegistry().FindOrAddHandle(UserId.AsShared());
+						AccountInfoOSSAdapter->UniqueNetId = UserId.AsShared();
+						AccountInfoOSSAdapter->LocalUserNum = LocalUserNum;
+						AccountInfoOSSAdapter->PlatformUserId = GetIdentityInterface()->GetPlatformUserIdFromLocalUserNum(LocalUserNum);
 					}
+				}
+				else
+				{
+					UE_LOG(LogOnlineServices, Warning, TEXT("[FAuthOSSAdapter::Login][%s] Failure: Failed to login to OSSv1 %s identity interface. Error %s"), *GetSubsystem().GetSubsystemName().ToString(), *Error);
+					FOnlineError V2Error = Errors::Unknown(); // TODO: V1 to V2 error conversion/error from string conversion
+					Op->SetError(MoveTemp(V2Error));
+				}
+			}
 
-				});
+			Promise.EmplaceValue();
+		});
 
-				GetIdentityInterface()->Login(LocalUserIndex, Credentials);
-			})
-			.Then([this](TOnlineAsyncOp<FAuthLogin>& Op, const TOnlineResult<FAuthLogin>& Result)
+		if (InAsyncOp.GetParams().CredentialsType == LoginCredentialsType::Auto)
+		{
+			GetIdentityInterface()->AutoLogin(InitialLocalUserNum);
+		}
+		else
+		{
+			FOnlineAccountCredentials Credentials;
+			Credentials.Type = InAsyncOp.GetParams().CredentialsType.ToString();
+			Credentials.Id = InAsyncOp.GetParams().CredentialsId;
+			Credentials.Token = InAsyncOp.GetParams().CredentialsToken.Get<FString>(); // TODO: handle binary token
+
+			GetIdentityInterface()->Login(InitialLocalUserNum, Credentials);
+		}
+		return Future;
+	})
+	// Step 4: Fetch dependent data.
+	.Then([this](TOnlineAsyncOp<FAuthLogin>& InAsyncOp)
+	{
+		const TSharedRef<FAccountInfoOSSAdapter>& AccountInfoOSSAdapter = GetOpDataChecked<TSharedRef<FAccountInfoOSSAdapter>>(InAsyncOp, AccountInfoKeyName);
+
+		if (TSharedPtr<FUserOnlineAccount> UserOnlineAccount = GetIdentityInterface()->GetUserAccount(*AccountInfoOSSAdapter->UniqueNetId))
+		{
+			AccountInfoOSSAdapter->Attributes.Emplace(AccountAttributeData::DisplayName, UserOnlineAccount->GetDisplayName());
+			return MakeFulfilledPromise<void>().GetFuture();
+		}
+		else
+		{
+			UE_LOG(LogOnlineServices, Warning, TEXT("[FAuthOSSAdapter::Login][%s] Failure: Failed to find UserOnlineAccount for account %s"), *GetSubsystem().GetSubsystemName().ToString(), *ToLogString(AccountInfoOSSAdapter->AccountId));
+			InAsyncOp.SetError(Errors::Unknown());
+
+			// Logout.
+			TPromise<void> Promise;
+			TFuture<void> Future = Promise.GetFuture();
+			MakeMulticastAdapter(this, GetIdentityInterface()->OnLogoutCompleteDelegates[AccountInfoOSSAdapter->LocalUserNum],
+			[this, AccountInfoOSSAdapter, Promise = MoveTemp(Promise)](int32 LocalUserNum, bool bWasSuccessful) mutable
 			{
-				if (Result.IsError())
-				{
-					Op.SetError(CopyTemp(Result.GetErrorValue()));
-				}
-				else
+				if (!bWasSuccessful)
 				{
-					Op.SetResult(CopyTemp(Result.GetOkValue()));
+					UE_LOG(LogOnlineServices, Warning, TEXT("[FAuthOSSAdapter::Login][%s] Failed to logout account [%s]."), *GetSubsystem().GetSubsystemName().ToString(), *ToLogString(AccountInfoOSSAdapter->AccountId));
 				}
-			})
-			.Enqueue(GetSerialQueue());
-	}
+				Promise.EmplaceValue();
+			});
+
+			GetIdentityInterface()->Logout(AccountInfoOSSAdapter->LocalUserNum);
+			return Future;
+		}
+	})
+	// Step 6: bookkeeping and notifications.
+	.Then([this](TOnlineAsyncOp<FAuthLogin>& InAsyncOp)
+	{
+		const TSharedRef<FAccountInfoOSSAdapter>& AccountInfoOSSAdapter = GetOpDataChecked<TSharedRef<FAccountInfoOSSAdapter>>(InAsyncOp, AccountInfoKeyName);
+
+		AccountInfoOSSAdapter->LoginStatus = ELoginStatus::LoggedIn;
+		AccountInfoRegistryOSSAdapter.Register(AccountInfoOSSAdapter);
+
+		UE_LOG(LogOnlineServices, Log, TEXT("[FAuthOSSAdapter::Login][%s] Successfully logged in as [%s]"), *GetSubsystem().GetSubsystemName().ToString(), *ToLogString(AccountInfoOSSAdapter->AccountId));
+		OnAuthLoginStatusChangedEvent.Broadcast(FAuthLoginStatusChanged{ AccountInfoOSSAdapter, AccountInfoOSSAdapter->LoginStatus });
+
+		InAsyncOp.SetResult(FAuthLogin::Result{ AccountInfoOSSAdapter });
+	})
+	.Enqueue(GetSerialQueue());
+
 	return Op->GetHandle();
 }
 
@@ -139,271 +301,304 @@ TOnlineAsyncOpHandle<FAuthLogout> FAuthOSSAdapter::Logout(FAuthLogout::Params&&
 {
 	TSharedRef<TOnlineAsyncOp<FAuthLogout>> Op = GetOp<FAuthLogout>(MoveTemp(Params));
 
-	if (!Op->IsReady())
+	// Step 1: Setup operation data.
+	Op->Then([this](TOnlineAsyncOp<FAuthLogout>& InAsyncOp)
 	{
-		int32 LocalUserIndex = GetLocalUserNum(Op->GetParams().LocalUserId);
+		const FAuthLogout::Params& Params = InAsyncOp.GetParams();
 
-		if (LocalUserIndex < 0 || LocalUserIndex >= MAX_LOCAL_PLAYERS)
+		TSharedPtr<FAccountInfoOSSAdapter> AccountInfoOSSAdapter = AccountInfoRegistryOSSAdapter.Find(Params.LocalUserId);
+		if (!AccountInfoOSSAdapter)
 		{
-			Op->SetError(Errors::InvalidParams());
-			return Op->GetHandle();
+			InAsyncOp.SetError(Errors::InvalidUser());
+			return;
 		}
 
-		Op->Then([this, LocalUserIndex](TOnlineAsyncOp<FAuthLogout>& Op, TPromise<TOnlineResult<FAuthLogout>>&& Result)
+		// Set user auth data on operation.
+		InAsyncOp.Data.Set<TSharedRef<FAccountInfoOSSAdapter>>(AccountInfoKeyName, AccountInfoOSSAdapter.ToSharedRef());
+	})
+	// Step 2: Logout the user.
+	.Then([this](TOnlineAsyncOp<FAuthLogout>& InAsyncOp)
+	{
+		const TSharedRef<FAccountInfoOSSAdapter>& AccountInfoOSSAdapter = GetOpDataChecked<TSharedRef<FAccountInfoOSSAdapter>>(InAsyncOp, AccountInfoKeyName);
+
+		TPromise<void> Promise;
+		TFuture<void> Future = Promise.GetFuture();
+		MakeMulticastAdapter(this, GetIdentityInterface()->OnLogoutCompleteDelegates[AccountInfoOSSAdapter->LocalUserNum],
+		[this, AccountInfoOSSAdapter, Promise = MoveTemp(Promise)](int32 LocalUserNum, bool bWasSuccessful) mutable
+		{
+			if (!bWasSuccessful)
 			{
-				MakeMulticastAdapter(this, GetIdentityInterface()->OnLogoutCompleteDelegates[LocalUserIndex],
-				[this, OnlineResult = MoveTemp(Result), LocalUserIndex](int32 LocalUserNum, bool bWasSuccessful) mutable
-				{
-					if (bWasSuccessful)
-					{
-						FAuthLogout::Result Result;
-						OnlineResult.EmplaceValue(MoveTemp(Result));
-					}
-					else
-					{
-						FOnlineError V2Error = Errors::Unknown(); // TODO: V1 to V2 error conversion/error from string conversion
-						OnlineResult.EmplaceValue(MoveTemp(V2Error));
-					}
-				});
+				UE_LOG(LogOnlineServices, Warning, TEXT("[FAuthOSSAdapter::Logout][%s] Failed to logout account [%s]."), *GetSubsystem().GetSubsystemName().ToString(), *ToLogString(AccountInfoOSSAdapter->AccountId));
+			}
+			Promise.EmplaceValue();
+		});
+
+		GetIdentityInterface()->Logout(AccountInfoOSSAdapter->LocalUserNum);
+		return Future;
+	})
+	// Step 3: bookkeeping and notifications.
+	.Then([this](TOnlineAsyncOp<FAuthLogout>& InAsyncOp)
+	{
+		const TSharedRef<FAccountInfoOSSAdapter>& AccountInfoOSSAdapter = GetOpDataChecked<TSharedRef<FAccountInfoOSSAdapter>>(InAsyncOp, AccountInfoKeyName);
+
+		UE_LOG(LogOnlineServices, Log, TEXT("[FAuthOSSAdapter::Logout][%s] Successfully logged out [%s]"), *GetSubsystem().GetSubsystemName().ToString(), *ToLogString(AccountInfoOSSAdapter->AccountId));
+		AccountInfoOSSAdapter->LoginStatus = ELoginStatus::NotLoggedIn;
+		OnAuthLoginStatusChangedEvent.Broadcast(FAuthLoginStatusChanged{ AccountInfoOSSAdapter, AccountInfoOSSAdapter->LoginStatus });
+		AccountInfoRegistryOSSAdapter.Unregister(AccountInfoOSSAdapter->AccountId);
+		InAsyncOp.SetResult(FAuthLogout::Result{});
+	})
+	.Enqueue(GetSerialQueue());
 
-				GetIdentityInterface()->Logout(LocalUserIndex);
-			})
-			.Then([this](TOnlineAsyncOp<FAuthLogout>& Op, const TOnlineResult<FAuthLogout>& Result)
-				{
-					if (Result.IsError())
-					{
-						Op.SetError(CopyTemp(Result.GetErrorValue()));
-					}
-					else
-					{
-						Op.SetResult(CopyTemp(Result.GetOkValue()));
-					}
-				})
-				.Enqueue(GetSerialQueue());
-	}
 	return Op->GetHandle();
 }
 
-TOnlineAsyncOpHandle<FAuthGenerateAuthToken> FAuthOSSAdapter::GenerateAuthToken(FAuthGenerateAuthToken::Params&& Params)
+TOnlineAsyncOpHandle<FAuthQueryExternalServerAuthTicket> FAuthOSSAdapter::QueryExternalServerAuthTicket(FAuthQueryExternalServerAuthTicket::Params&& Params)
 {
-	TSharedRef<TOnlineAsyncOp<FAuthGenerateAuthToken>> Op = GetOp<FAuthGenerateAuthToken>(MoveTemp(Params));
+	TOnlineAsyncOpRef<FAuthQueryExternalServerAuthTicket> Op = GetOp<FAuthQueryExternalServerAuthTicket>(MoveTemp(Params));
 
-	if (!Op->IsReady())
+	// Step 1: Setup operation data.
+	Op->Then([this](TOnlineAsyncOp<FAuthQueryExternalServerAuthTicket>& InAsyncOp)
 	{
-		const FUniqueNetIdRef UniqueNetId = GetUniqueNetId(Op->GetParams().LocalUserId);
-		if (!UniqueNetId->IsValid())
-		{
-			Op->SetError(Errors::InvalidParams());
-			return Op->GetHandle();
-		}
+		const FAuthQueryExternalServerAuthTicket::Params& Params = InAsyncOp.GetParams();
 
-		IOnlineIdentityPtr Identity = GetIdentityInterface();
-		if (Identity->GetLoginStatus(*UniqueNetId) == ::ELoginStatus::LoggedIn)
+		// Check for supported translator
+		TDefaultErrorResultInternal<TSharedRef<FExternalServerAuthTicketTranslationTraits>> TranslationTraits = GetExternalServerAuthTicketTranslationTraits(GetSubsystem().GetSubsystemName());
+		if (!TranslationTraits.IsOk())
 		{
-			Op->Then([this](TOnlineAsyncOp<FAuthGenerateAuthToken>& Op)
-			{
-				if (IOnlineIdentityPtr Identity = GetIdentityInterface())
-				{
-					const FUniqueNetIdRef UniqueNetId = GetUniqueNetId(Op.GetParams().LocalUserId);
-					const int32 LocalUserNum = Identity->GetLocalUserNumFromPlatformUserId(Identity->GetPlatformUserIdFromUniqueNetId(*UniqueNetId));
+			InAsyncOp.SetError(MoveTemp(TranslationTraits.GetErrorValue()));
+			return;
+		}
 
-					Identity->GetLinkedAccountAuthToken(LocalUserNum, *MakeDelegateAdapter(this, 
-					[this, WeakOp = Op.AsWeak()](int32 LocalUserNum, bool bWasSuccessful, const FExternalAuthToken& AuthToken)
-					{
-						TSharedPtr<TOnlineAsyncOp<FAuthGenerateAuthToken>> Op = WeakOp.Pin();
-						if (!Op)
-						{
-							return;
-						}
+		// Set translator traits on operation.
+		InAsyncOp.Data.Set<TSharedRef<FExternalServerAuthTicketTranslationTraits>>(ExternalServerAuthTicketTranslationTraitsKeyName, TranslationTraits.GetOkValue());
 
-						if (bWasSuccessful && AuthToken.IsValid())
-						{
-							FAuthGenerateAuthToken::Result Result;
-							Result.Type = LexToString(Services.GetServicesProvider());
-							if (AuthToken.HasTokenString())
-							{
-								Result.Token.Emplace<FString>(AuthToken.TokenString);
-							}
-							else if (AuthToken.HasTokenData())
-							{
-								Result.Token.Emplace<TArray<uint8>>(AuthToken.TokenData);
-							}
-							else
-							{
-								checkNoEntry();
-							}
-							Op->SetResult(MoveTemp(Result));
-						}
-						else
-						{
-							Op->SetError(Errors::Unknown());
-						}
-					}));
-				}
-				else
-				{
-					Op.SetError(Errors::Unknown());
-				}
-			})
-			.Enqueue(GetSerialQueue());
-		}
-		else
+		// Look up logged in user.
+		TSharedPtr<FAccountInfoOSSAdapter> AccountInfoOSSAdapter = AccountInfoRegistryOSSAdapter.Find(Params.LocalUserId);
+		if (!AccountInfoOSSAdapter)
 		{
-			Op->SetError(Errors::InvalidAuth());
+			InAsyncOp.SetError(Errors::InvalidUser());
+			return;
 		}
-	}
+
+		// Set user auth data on operation.
+		InAsyncOp.Data.Set<TSharedRef<FAccountInfoOSSAdapter>>(AccountInfoKeyName, AccountInfoOSSAdapter.ToSharedRef());
+	})
+	// Step 2: Fetch OSSv1 auth ticket data and signal result.
+	.Then([this](TOnlineAsyncOp<FAuthQueryExternalServerAuthTicket>& InAsyncOp)
+	{
+		const TSharedRef<FAccountInfoOSSAdapter>& AccountInfoOSSAdapter = GetOpDataChecked<TSharedRef<FAccountInfoOSSAdapter>>(InAsyncOp, AccountInfoKeyName);
+		const TSharedRef<FExternalServerAuthTicketTranslationTraits>& ExternalServerAuthTicketTranslationTraits = GetOpDataChecked<TSharedRef<FExternalServerAuthTicketTranslationTraits>>(InAsyncOp, ExternalAuthTokenTranslationTraitsKeyName);
+
+		InAsyncOp.SetResult(FAuthQueryExternalServerAuthTicket::Result{FExternalServerAuthTicket{ExternalServerAuthTicketTranslationTraits->ExternalTicketType, GetIdentityInterface()->GetAuthToken(AccountInfoOSSAdapter->LocalUserNum)}});
+	})
+	.Enqueue(GetSerialQueue());
 
 	return Op->GetHandle();
 }
 
-TOnlineAsyncOpHandle<FAuthGenerateAuthCode> FAuthOSSAdapter::GenerateAuthCode(FAuthGenerateAuthCode::Params&& Params)
+TOnlineAsyncOpHandle<FAuthQueryExternalAuthToken> FAuthOSSAdapter::QueryExternalAuthToken(FAuthQueryExternalAuthToken::Params&& Params)
 {
-	TSharedRef<TOnlineAsyncOp<FAuthGenerateAuthCode>> Op = GetOp<FAuthGenerateAuthCode>(MoveTemp(Params));
+	TOnlineAsyncOpRef<FAuthQueryExternalAuthToken> Op = GetOp<FAuthQueryExternalAuthToken>(MoveTemp(Params));
 
-	if (!Op->IsReady())
+	// Step 1: Setup operation data.
+	Op->Then([this](TOnlineAsyncOp<FAuthQueryExternalAuthToken>& InAsyncOp)
 	{
-		const FUniqueNetIdRef UniqueNetId = GetUniqueNetId(Op->GetParams().LocalUserId);
-		if (!UniqueNetId->IsValid())
+		const FAuthQueryExternalAuthToken::Params& Params = InAsyncOp.GetParams();
+
+		// Check for supported translator
+		const FExternalAuthTokenTranslationTraits* TranslationTraits = GetExternalAuthTokenTranslationTraits(GetSubsystem().GetSubsystemName(), Params.Method);
+		if (TranslationTraits == nullptr)
 		{
-			Op->SetError(Errors::InvalidParams());
-			return Op->GetHandle();
+			InAsyncOp.SetError(Errors::NotImplemented());
+			return;
 		}
 
-		IOnlineIdentityPtr Identity = GetIdentityInterface();
-		if (Identity->GetLoginStatus(*UniqueNetId) == ::ELoginStatus::LoggedIn)
+		// Set translator traits on operation.
+		InAsyncOp.Data.Set<const FExternalAuthTokenTranslationTraits*>(ExternalAuthTokenTranslationTraitsKeyName, TranslationTraits);
+
+		// Look up logged in user.
+		TSharedPtr<FAccountInfoOSSAdapter> AccountInfoOSSAdapter = AccountInfoRegistryOSSAdapter.Find(Params.LocalUserId);
+		if (!AccountInfoOSSAdapter)
 		{
-			const int32 LocalUserNum = Identity->GetLocalUserNumFromPlatformUserId(Identity->GetPlatformUserIdFromUniqueNetId(*UniqueNetId));
+			InAsyncOp.SetError(Errors::InvalidUser());
+			return;
+		}
+
+		// Set user auth data on operation.
+		InAsyncOp.Data.Set<TSharedRef<FAccountInfoOSSAdapter>>(AccountInfoKeyName, AccountInfoOSSAdapter.ToSharedRef());
+	})
+	// Step 2: Fetch OSSv1 token data.
+	.Then([this](TOnlineAsyncOp<FAuthQueryExternalAuthToken>& InAsyncOp)
+	{
+		const TSharedRef<FAccountInfoOSSAdapter>& AccountInfoOSSAdapter = GetOpDataChecked<TSharedRef<FAccountInfoOSSAdapter>>(InAsyncOp, AccountInfoKeyName);
+		const FExternalAuthTokenTranslationTraits* ExternalAuthTokenTranslationTraits = GetOpDataChecked<const FExternalAuthTokenTranslationTraits*>(InAsyncOp, ExternalAuthTokenTranslationTraitsKeyName);
 
-			FString Token = Identity->GetAuthToken(LocalUserNum);
-			if (Token.IsEmpty())
+		TPromise<FString> Promise;
+		TFuture<FString> Future = Promise.GetFuture();
+
+		if (EnumHasAnyFlags(ExternalAuthTokenTranslationTraits->Flags, EExternalAuthTokenTranslationFlags::AuthToken))
+		{
+			Promise.EmplaceValue(GetIdentityInterface()->GetAuthToken(AccountInfoOSSAdapter->LocalUserNum));
+		}
+		else if (EnumHasAnyFlags(ExternalAuthTokenTranslationTraits->Flags, EExternalAuthTokenTranslationFlags::TokenString | EExternalAuthTokenTranslationFlags::TokenBinary))
+		{
+			GetIdentityInterface()->GetLinkedAccountAuthToken(AccountInfoOSSAdapter->LocalUserNum, *MakeDelegateAdapter(this,
+			[this, WeakOp = InAsyncOp.AsWeak(), Promise = MoveTemp(Promise), ExternalAuthTokenTranslationTraits](int32 LocalUserNum, bool bWasSuccessful, const ::FExternalAuthToken& AuthToken) mutable
 			{
-				Op->Then([this](TOnlineAsyncOp<FAuthGenerateAuthCode>& Op)
+				if (TSharedPtr<TOnlineAsyncOp<FAuthQueryExternalAuthToken>> Op = WeakOp.Pin())
 				{
-					// AutoLogin to generate the token. Token generation is a side effect of calling login, but you can be logged in without calling login.
-					FOnLoginCompleteDelegate LoginCompleteDelegate = FOnLoginCompleteDelegate::CreateLambda([this, WeakOp = Op.AsWeak()](int32 LocalUserNum, bool bWasSuccessful, const FUniqueNetId& UserId, const FString& Error)
+					if (bWasSuccessful)
 					{
-						TSharedPtr<TOnlineAsyncOp<FAuthGenerateAuthCode>> Op = WeakOp.Pin();
-						if (!Op)
+						if (EnumHasAnyFlags(ExternalAuthTokenTranslationTraits->Flags, EExternalAuthTokenTranslationFlags::TokenString))
 						{
-							return;
+							if (AuthToken.HasTokenString())
+							{
+								Promise.EmplaceValue(AuthToken.TokenString);
+							}
+							else
+							{
+								UE_LOG(LogOnlineServices, Warning, TEXT("[FAuthOSSAdapter::QueryExternalAuthToken][%s] Failed: Token string not found."), *GetSubsystem().GetSubsystemName().ToString());
+								Op->SetError(Errors::Unknown());
+								Promise.EmplaceValue();
+							}
 						}
-
-						IOnlineIdentityPtr Identity = GetIdentityInterface();
-						if (bWasSuccessful && Identity)
+						else if (EnumHasAnyFlags(ExternalAuthTokenTranslationTraits->Flags, EExternalAuthTokenTranslationFlags::TokenBinary))
 						{
-							if (const FDelegateHandle* LoginCompleteHandle = Op->Data.Get<FDelegateHandle>(TEXT("LoginCompleteHandle")))
+							if (AuthToken.HasTokenData())
 							{
-								Identity->ClearOnLoginCompleteDelegate_Handle(LocalUserNum, *const_cast<FDelegateHandle*>(LoginCompleteHandle));
+								// Convert binary to hex representation.
+								Promise.EmplaceValue(FString::FromHexBlob(AuthToken.TokenData.GetData(), AuthToken.TokenData.Num()));
 							}
-
-							FString Token = Identity->GetAuthToken(LocalUserNum);
-							if (!Token.IsEmpty())
+							else
 							{
-								FAuthGenerateAuthCode::Result Result;
-								Result.Code = MoveTemp(Token);
-								Op->SetResult(MoveTemp(Result));
-								return;
+								UE_LOG(LogOnlineServices, Warning, TEXT("[FAuthOSSAdapter::QueryExternalAuthToken][%s] Failed: Token binary data not found."), *GetSubsystem().GetSubsystemName().ToString());
+								Op->SetError(Errors::Unknown());
+								Promise.EmplaceValue();
 							}
 						}
+						else
+						{
+							checkNoEntry();
+							Op->SetError(Errors::Unknown());
+							Promise.EmplaceValue();
+						}
+					}
+					else
+					{
+						UE_LOG(LogOnlineServices, Warning, TEXT("[FAuthOSSAdapter::QueryExternalAuthToken][%s] GetLinkedAccountAuthToken failed."), *GetSubsystem().GetSubsystemName().ToString());
 						Op->SetError(Errors::Unknown());
-					});
-
-
-					IOnlineIdentityPtr Identity = GetIdentityInterface();
-					const FUniqueNetIdRef UniqueNetId = GetUniqueNetId(Op.GetParams().LocalUserId);
-					const int32 LocalUserNum = Identity->GetLocalUserNumFromPlatformUserId(Identity->GetPlatformUserIdFromUniqueNetId(*UniqueNetId));
-					FDelegateHandle LoginCompleteHandle = Identity->AddOnLoginCompleteDelegate_Handle(LocalUserNum, LoginCompleteDelegate);
-					Op.Data.Set(TEXT("LoginCompleteHandle"), MoveTemp(LoginCompleteHandle));
-					Identity->AutoLogin(LocalUserNum);
-				})
-				.Enqueue(GetSerialQueue());
-			}
-			else
-			{
-				FAuthGenerateAuthCode::Result Result;
-				Result.Code = MoveTemp(Token);
-				Op->SetResult(MoveTemp(Result));
-			}
+						Promise.EmplaceValue();
+					}
+				}
+				else
+				{
+					Promise.EmplaceValue();
+				}
+			}));
 		}
 		else
 		{
-			Op->SetError(Errors::InvalidAuth());
+			// Getting here means the traits were setup incorrectly.
+			InAsyncOp.SetError(Errors::Unknown());
+			Promise.EmplaceValue();
 		}
-	}
+
+		return Future;
+	})
+	// Step 3: Signal valid result.
+	.Then([this](TOnlineAsyncOp<FAuthQueryExternalAuthToken>& InAsyncOp, FString&& ResolvedExternalAuthToken)
+	{
+		const FExternalAuthTokenTranslationTraits* ExternalAuthTokenTranslationTraits = GetOpDataChecked<const FExternalAuthTokenTranslationTraits*>(InAsyncOp, ExternalAuthTokenTranslationTraitsKeyName);
+		InAsyncOp.SetResult(FAuthQueryExternalAuthToken::Result{FExternalAuthToken{ExternalAuthTokenTranslationTraits->ExternalLoginType, MoveTemp(ResolvedExternalAuthToken)}});
+	})
+	.Enqueue(GetSerialQueue());
 
 	return Op->GetHandle();
 }
 
-TOnlineResult<FAuthGetAccountByPlatformUserId> FAuthOSSAdapter::GetAccountByPlatformUserId(FAuthGetAccountByPlatformUserId::Params&& Params)
+FUniqueNetIdPtr FAuthOSSAdapter::GetUniqueNetId(FOnlineAccountIdHandle AccountIdHandle) const
 {
-	int32 LocalUserIndex = FPlatformMisc::GetUserIndexForPlatformUser(Params.PlatformUserId);
-
-	if (LocalUserIndex < 0 || LocalUserIndex >= MAX_LOCAL_PLAYERS)
-	{
-		return TOnlineResult<FAuthGetAccountByPlatformUserId>(Errors::InvalidUser());
-	}
-
-	FUniqueNetIdPtr UniqueNetId = GetIdentityInterface()->GetUniquePlayerId(LocalUserIndex);
-
-	if (!UniqueNetId.IsValid())
-	{
-		return TOnlineResult<FAuthGetAccountByPlatformUserId>(Errors::InvalidUser());
-	}
-
-	FOnlineAccountIdHandle Handle = static_cast<FOnlineServicesOSSAdapter&>(Services).GetAccountIdRegistry().FindOrAddHandle(UniqueNetId.ToSharedRef());
-	FAuthGetAccountByPlatformUserId::Result Result = { MakeShared<FAccountInfo>() };
-	Result.AccountInfo->PlatformUserId = Params.PlatformUserId;
-	Result.AccountInfo->UserId = Handle;
-	// v1 and v2 login status have the same values, so we can just cast here
-	Result.AccountInfo->LoginStatus = static_cast<ELoginStatus>(GetIdentityInterface()->GetLoginStatus(*UniqueNetId));
-	Result.AccountInfo->DisplayName = GetIdentityInterface()->GetPlayerNickname(*UniqueNetId);
-
-	return TOnlineResult<FAuthGetAccountByPlatformUserId>(MoveTemp(Result));
+	return GetOnlineServicesOSSAdapter().GetAccountIdRegistry().GetIdValue(AccountIdHandle);
 }
 
-TOnlineResult<FAuthGetAccountByAccountId> FAuthOSSAdapter::GetAccountByAccountId(FAuthGetAccountByAccountId::Params&& Params)
+FOnlineAccountIdHandle FAuthOSSAdapter::GetAccountIdHandle(const FUniqueNetIdRef& UniqueNetId) const
 {
-	FUniqueNetIdRef UniqueNetId = GetUniqueNetId(Params.LocalUserId);
-
-	if (!UniqueNetId->IsValid())
-	{
-		return TOnlineResult<FAuthGetAccountByAccountId>(Errors::InvalidUser());
-	}
+	return GetOnlineServicesOSSAdapter().GetAccountIdRegistry().FindOrAddHandle(UniqueNetId);
+}
 
-	FOnlineAccountIdHandle Handle = static_cast<FOnlineServicesOSSAdapter&>(Services).GetAccountIdRegistry().FindOrAddHandle(UniqueNetId);
-	FAuthGetAccountByAccountId::Result Result = { MakeShared<FAccountInfo>() };
-	Result.AccountInfo->PlatformUserId = GetIdentityInterface()->GetPlatformUserIdFromUniqueNetId(*UniqueNetId);
-	Result.AccountInfo->UserId = Handle;
-	// v1 and v2 login status have the same values, so we can just cast here
-	Result.AccountInfo->LoginStatus = static_cast<ELoginStatus>(GetIdentityInterface()->GetLoginStatus(*UniqueNetId));
-	Result.AccountInfo->DisplayName = GetIdentityInterface()->GetPlayerNickname(*UniqueNetId);
+int32 FAuthOSSAdapter::GetLocalUserNum(FOnlineAccountIdHandle AccountIdHandle) const
+{
+	TSharedPtr<FAccountInfoOSSAdapter> AccountInfoOSSAdapter = AccountInfoRegistryOSSAdapter.Find(AccountIdHandle);
+	return AccountInfoOSSAdapter ? AccountInfoOSSAdapter->LocalUserNum : INDEX_NONE;
+}
 
-	return TOnlineResult<FAuthGetAccountByAccountId>(MoveTemp(Result));
+const FAccountInfoRegistry& FAuthOSSAdapter::GetAccountInfoRegistry() const
+{
+	return AccountInfoRegistryOSSAdapter;
 }
 
-FUniqueNetIdRef FAuthOSSAdapter::GetUniqueNetId(FOnlineAccountIdHandle AccountIdHandle) const
+const FOnlineServicesOSSAdapter& FAuthOSSAdapter::GetOnlineServicesOSSAdapter() const
 {
-	return static_cast<FOnlineServicesOSSAdapter&>(Services).GetAccountIdRegistry().GetIdValue(AccountIdHandle);
+	return const_cast<FAuthOSSAdapter*>(this)->GetOnlineServicesOSSAdapter();
 }
 
-FOnlineAccountIdHandle UE::Online::FAuthOSSAdapter::GetAccountIdHandle(FUniqueNetIdRef UniqueNetId) const
+FOnlineServicesOSSAdapter& FAuthOSSAdapter::GetOnlineServicesOSSAdapter()
 {
-	return static_cast<FOnlineServicesOSSAdapter&>(Services).GetAccountIdRegistry().FindOrAddHandle(UniqueNetId);
+	return static_cast<FOnlineServicesOSSAdapter&>(Services);
 }
 
-int32 FAuthOSSAdapter::GetLocalUserNum(FOnlineAccountIdHandle AccountIdHandle) const
+const IOnlineSubsystem& FAuthOSSAdapter::GetSubsystem() const
 {
-	FUniqueNetIdRef UniqueNetId = GetUniqueNetId(AccountIdHandle);
-	if (UniqueNetId->IsValid())
-	{
-		return GetIdentityInterface()->GetLocalUserNumFromPlatformUserId(GetIdentityInterface()->GetPlatformUserIdFromUniqueNetId(*UniqueNetId));
-	}
-	return INDEX_NONE;
+	return GetOnlineServicesOSSAdapter().GetSubsystem();
 }
 
 IOnlineIdentityPtr FAuthOSSAdapter::GetIdentityInterface() const
 {
-	return static_cast<FOnlineServicesOSSAdapter&>(Services).GetSubsystem().GetIdentityInterface();
+	return GetSubsystem().GetIdentityInterface();
+}
+
+TOnlineAsyncOpHandle<FAuthOSSAdapter::FHandleLoginStatusChangedImpl> FAuthOSSAdapter::HandleLoginStatusChangedImplOp(FHandleLoginStatusChangedImpl::Params&& Params)
+{
+	TOnlineAsyncOpRef<FHandleLoginStatusChangedImpl> Op = GetOp<FHandleLoginStatusChangedImpl>(MoveTemp(Params));
+
+	// Step 1: Set up operation data.
+	Op->Then([this](TOnlineAsyncOp<FHandleLoginStatusChangedImpl>& InAsyncOp)
+	{
+		const FHandleLoginStatusChangedImpl::Params& Params = InAsyncOp.GetParams();
+		TSharedPtr<FAccountInfoOSSAdapter> AccountInfoOSSAdapter = AccountInfoRegistryOSSAdapter.Find(Params.AccountId);
+		if (!AccountInfoOSSAdapter)
+		{
+			InAsyncOp.SetError(Errors::InvalidUser());
+			return;
+		}
+
+		// Don't send duplicate notification during login
+		if (AccountInfoOSSAdapter->LoginStatus == Params.NewLoginStatus)
+		{
+			UE_LOG(LogOnlineServices, Log, TEXT("[FAuthOSSAdapter::HandleLoginStatusChangedImplOp][%s] Login status has not changed. Ignoring event. [%s] Current Status: %s"),
+				*GetSubsystem().GetSubsystemName().ToString(), *ToLogString(AccountInfoOSSAdapter->AccountId), *ToLogString(Params.NewLoginStatus));
+			InAsyncOp.SetResult(FHandleLoginStatusChangedImpl::Result{});
+			return;
+		}
+
+		// Set user auth data on operation.
+		InAsyncOp.Data.Set<TSharedRef<FAccountInfoOSSAdapter>>(AccountInfoKeyName, AccountInfoOSSAdapter.ToSharedRef());
+	})
+	// Step 2: Update status and notify.
+	.Then([this](TOnlineAsyncOp<FHandleLoginStatusChangedImpl>& InAsyncOp)
+	{
+		const FHandleLoginStatusChangedImpl::Params& Params = InAsyncOp.GetParams();
+		const TSharedRef<FAccountInfoOSSAdapter>& AccountInfoOSSAdapter = GetOpDataChecked<TSharedRef<FAccountInfoOSSAdapter>>(InAsyncOp, AccountInfoKeyName);
+
+		UE_LOG(LogOnlineServices, Log, TEXT("[FAuthOSSAdapter::HandleLoginStatusChangedImplOp][%s] Login status changed for account [%s]: %s => %s."),
+			*GetSubsystem().GetSubsystemName().ToString(), *ToLogString(AccountInfoOSSAdapter->AccountId), *ToLogString(AccountInfoOSSAdapter->LoginStatus), *ToLogString(Params.NewLoginStatus));
+		AccountInfoOSSAdapter->LoginStatus = Params.NewLoginStatus;
+		OnAuthLoginStatusChangedEvent.Broadcast(FAuthLoginStatusChanged{ AccountInfoOSSAdapter, AccountInfoOSSAdapter->LoginStatus });
+		InAsyncOp.SetResult(FHandleLoginStatusChangedImpl::Result{});
+	})
+	.Enqueue(GetSerialQueue());
+
+	return Op->GetHandle();
 }
 
 /* UE::Online */ }
diff --git a/Engine/Plugins/Online/Experimental/OnlineServicesOSSAdapter/Source/Private/Online/PrivilegesOSSAdapter.cpp b/Engine/Plugins/Online/Experimental/OnlineServicesOSSAdapter/Source/Private/Online/PrivilegesOSSAdapter.cpp
--- a/Engine/Plugins/Online/Experimental/OnlineServicesOSSAdapter/Source/Private/Online/PrivilegesOSSAdapter.cpp
+++ b/Engine/Plugins/Online/Experimental/OnlineServicesOSSAdapter/Source/Private/Online/PrivilegesOSSAdapter.cpp
@@ -27,8 +27,8 @@ TOnlineAsyncOpHandle<FQueryUserPrivilege> FPrivilegesOSSAdapter::QueryUserPrivil
 	if (!Op->IsReady())
 	{
 		FAuthOSSAdapter* Auth = Services.Get<FAuthOSSAdapter>();
-		FUniqueNetIdRef UniqueNetId = Auth->GetUniqueNetId(Op->GetParams().LocalUserId);
-		if (!UniqueNetId->IsValid())
+		FUniqueNetIdPtr UniqueNetId = Auth->GetUniqueNetId(Op->GetParams().LocalUserId);
+		if (!UniqueNetId.IsValid())
 		{
 			Op->SetError(Errors::InvalidUser());
 			return Op->GetHandle();
diff --git a/Engine/Plugins/Online/Experimental/OnlineServicesOSSAdapter/Source/Public/Online/AuthOSSAdapter.h b/Engine/Plugins/Online/Experimental/OnlineServicesOSSAdapter/Source/Public/Online/AuthOSSAdapter.h
--- a/Engine/Plugins/Online/Experimental/OnlineServicesOSSAdapter/Source/Public/Online/AuthOSSAdapter.h
+++ b/Engine/Plugins/Online/Experimental/OnlineServicesOSSAdapter/Source/Public/Online/AuthOSSAdapter.h
@@ -6,11 +6,34 @@
 
 #include "OnlineSubsystemTypes.h"
 
+class IOnlineSubsystem;
 class IOnlineIdentity;
 using IOnlineIdentityPtr = TSharedPtr<IOnlineIdentity>;
 
 namespace UE::Online {
 
+class FOnlineServicesOSSAdapter;
+
+struct FAccountInfoOSSAdapter final : public FAccountInfo
+{
+	FUniqueNetIdPtr UniqueNetId;
+	int32 LocalUserNum = INDEX_NONE;
+};
+
+class FAccountInfoRegistryOSSAdapter final : public FAccountInfoRegistry
+{
+public:
+	using Super = FAccountInfoRegistry;
+
+	virtual ~FAccountInfoRegistryOSSAdapter() = default;
+
+	TSharedPtr<FAccountInfoOSSAdapter> Find(FPlatformUserId PlatformUserId) const;
+	TSharedPtr<FAccountInfoOSSAdapter> Find(FOnlineAccountIdHandle AccountIdHandle) const;
+
+	void Register(const TSharedRef<FAccountInfoOSSAdapter>&UserAuthData);
+	void Unregister(FOnlineAccountIdHandle AccountId);
+};
+
 class FAuthOSSAdapter : public FAuthCommon
 {
 public:
@@ -25,18 +48,40 @@ public:
 	// IAuth
 	virtual TOnlineAsyncOpHandle<FAuthLogin> Login(FAuthLogin::Params&& Params) override;
 	virtual TOnlineAsyncOpHandle<FAuthLogout> Logout(FAuthLogout::Params&& Params) override;
-	virtual TOnlineAsyncOpHandle<FAuthGenerateAuthToken> GenerateAuthToken(FAuthGenerateAuthToken::Params&& Params) override;
-	virtual TOnlineAsyncOpHandle<FAuthGenerateAuthCode> GenerateAuthCode(FAuthGenerateAuthCode::Params&& Params) override;
-	virtual TOnlineResult<FAuthGetAccountByPlatformUserId> GetAccountByPlatformUserId(FAuthGetAccountByPlatformUserId::Params&& Params) override;
-	virtual TOnlineResult<FAuthGetAccountByAccountId> GetAccountByAccountId(FAuthGetAccountByAccountId::Params&& Params) override;
+	virtual TOnlineAsyncOpHandle<FAuthQueryExternalServerAuthTicket> QueryExternalServerAuthTicket(FAuthQueryExternalServerAuthTicket::Params&& Params) override;
+	virtual TOnlineAsyncOpHandle<FAuthQueryExternalAuthToken> QueryExternalAuthToken(FAuthQueryExternalAuthToken::Params&& Params) override;
 
-	FUniqueNetIdRef GetUniqueNetId(FOnlineAccountIdHandle AccountIdHandle) const;
-	FOnlineAccountIdHandle GetAccountIdHandle(FUniqueNetIdRef UniqueNetId) const;
+	FUniqueNetIdPtr GetUniqueNetId(FOnlineAccountIdHandle AccountIdHandle) const;
+	FOnlineAccountIdHandle GetAccountIdHandle(const FUniqueNetIdRef& UniqueNetId) const;
 	int32 GetLocalUserNum(FOnlineAccountIdHandle AccountIdHandle) const;
 
+protected:
+	virtual const FAccountInfoRegistry& GetAccountInfoRegistry() const override;
+
+	const FOnlineServicesOSSAdapter& GetOnlineServicesOSSAdapter() const;
+	FOnlineServicesOSSAdapter& GetOnlineServicesOSSAdapter();
+	const IOnlineSubsystem& GetSubsystem() const;
 	IOnlineIdentityPtr GetIdentityInterface() const;
 
-protected:
+	struct FHandleLoginStatusChangedImpl
+	{
+		static constexpr TCHAR Name[] = TEXT("HandleLoginStatusChangedImpl");
+
+		struct Params
+		{
+			FPlatformUserId PlatformUserId;
+			FOnlineAccountIdHandle AccountId;
+			ELoginStatus NewLoginStatus;
+		};
+
+		struct Result
+		{
+		};
+	};
+
+	TOnlineAsyncOpHandle<FHandleLoginStatusChangedImpl> HandleLoginStatusChangedImplOp(FHandleLoginStatusChangedImpl::Params&& Params);
+
+	FAccountInfoRegistryOSSAdapter AccountInfoRegistryOSSAdapter;
 	FDelegateHandle OnLoginStatusChangedHandle[MAX_LOCAL_PLAYERS];
 };
 
diff --git a/Engine/Plugins/Online/Experimental/OnlineServicesOSSAdapter/Source/Public/Online/OnlineIdOSSAdapter.h b/Engine/Plugins/Online/Experimental/OnlineServicesOSSAdapter/Source/Public/Online/OnlineIdOSSAdapter.h
--- a/Engine/Plugins/Online/Experimental/OnlineServicesOSSAdapter/Source/Public/Online/OnlineIdOSSAdapter.h
+++ b/Engine/Plugins/Online/Experimental/OnlineServicesOSSAdapter/Source/Public/Online/OnlineIdOSSAdapter.h
@@ -24,6 +24,11 @@
 	FOnlineAccountIdHandle FindOrAddHandle(const FUniqueNetIdRef& IdValue)
 	{
 		FOnlineAccountIdHandle Handle;
+		if (!ensure(IdValue->IsValid()))
+		{
+			return Handle;
+		}
+
 		// Take a read lock and check if we already have a handle
 		{
 			FReadScopeLock ReadLock(Lock);
@@ -54,7 +59,7 @@
 	}
 
 	// Returns a copy as it's not thread safe to return a pointer/ref to an element of an array that can be relocated by another thread.
-	FUniqueNetIdRef GetIdValue(const FOnlineAccountIdHandle Handle) const
+	FUniqueNetIdPtr GetIdValue(const FOnlineAccountIdHandle Handle) const
 	{
 		if (Handle.GetOnlineServicesType() == OnlineServicesType
 			&& Handle.IsValid()
@@ -63,15 +68,12 @@
 			FReadScopeLock ReadLock(Lock);
 			return IdValues[Handle.GetHandle() - 1];
 		}
-		return FUniqueNetIdString::EmptyId();
+		return FUniqueNetIdPtr();
 	}
 
 	FUniqueNetIdRef GetIdValueChecked(const FOnlineAccountIdHandle Handle) const
 	{
-		check(Handle.GetOnlineServicesType() == OnlineServicesType
-			&& Handle.IsValid()
-			&& IdValues.IsValidIndex(Handle.GetHandle() - 1));
-		return GetIdValue(Handle);
+		return GetIdValue(Handle).ToSharedRef();
 	}
 
 	// Begin IOnlineAccountIdRegistry
diff --git a/Engine/Plugins/Online/Experimental/OnlineServicesOSSAdapter/Source/Public/Online/OnlineServicesOSSAdapter.h b/Engine/Plugins/Online/Experimental/OnlineServicesOSSAdapter/Source/Public/Online/OnlineServicesOSSAdapter.h
--- a/Engine/Plugins/Online/Experimental/OnlineServicesOSSAdapter/Source/Public/Online/OnlineServicesOSSAdapter.h
+++ b/Engine/Plugins/Online/Experimental/OnlineServicesOSSAdapter/Source/Public/Online/OnlineServicesOSSAdapter.h
@@ -21,7 +21,8 @@ public:
 	ONLINESERVICESOSSADAPTER_API virtual void Initialize() override;
 	virtual EOnlineServices GetServicesProvider() const override { return ServicesType; }
 
-	IOnlineSubsystem& GetSubsystem() { return *Subsystem; }
+	IOnlineSubsystem& GetSubsystem() const { return *Subsystem; }
+	FOnlineUniqueNetIdRegistry& GetAccountIdRegistry() const { return *AccountIdRegistry; }
 	FOnlineUniqueNetIdRegistry& GetAccountIdRegistry() { return *AccountIdRegistry; }
 
 protected:
