- UserInfo OSS Adapter implementation

diff --git a/Engine/Plugins/Online/Experimental/OnlineServicesOSSAdapter/Source/Private/Online/UserInfoOSSAdapter.cpp b/Engine/Plugins/Online/Experimental/OnlineServicesOSSAdapter/Source/Private/Online/UserInfoOSSAdapter.cpp
new file mode 100644
--- /dev/null
+++ b/Engine/Plugins/Online/Experimental/OnlineServicesOSSAdapter/Source/Private/Online/UserInfoOSSAdapter.cpp
@@ -0,0 +1,125 @@
+// Copyright Epic Games, Inc. All Rights Reserved.
+
+#include "Online/UserInfoOSSAdapter.h"
+
+#include "Interfaces/OnlineUserInterface.h"
+#include "Online/AuthOSSAdapter.h"
+#include "Online/DelegateAdapter.h"
+#include "Online/OnlineErrorDefinitions.h"
+#include "Online/OnlineIdOSSAdapter.h"
+#include "Online/OnlineServicesOSSAdapter.h"
+
+#include "OnlineSubsystem.h"
+
+namespace UE::Online {
+
+void FUserInfoOSSAdapter::Initialize()
+{
+	Super::Initialize();
+
+	IOnlineSubsystem& Subsystem = static_cast<FOnlineServicesOSSAdapter&>(Services).GetSubsystem();
+
+	UserInterface = Subsystem.GetUserInterface();
+	check(UserInterface);
+}
+
+TOnlineAsyncOpHandle<FQueryUserInfo> FUserInfoOSSAdapter::QueryUserInfo(FQueryUserInfo::Params&& InParams)
+{
+	TOnlineAsyncOpRef<FQueryUserInfo> Op = GetJoinableOp<FQueryUserInfo>(MoveTemp(InParams));
+	const FQueryUserInfo::Params& Params = Op->GetParams();
+
+	if (Params.UserIds.IsEmpty())
+	{
+		Op->SetError(Errors::InvalidParams());
+		return Op->GetHandle();
+	}
+
+	Op->Then([this](TOnlineAsyncOp<FQueryUserInfo>& Op)
+	{
+		const FQueryUserInfo::Params& Params = Op.GetParams();
+		const FAuthOSSAdapter* AuthInterface = Services.Get<FAuthOSSAdapter>();
+
+		const int32 LocalUserNum = AuthInterface->GetLocalUserNum(Params.LocalUserId);
+		if (LocalUserNum == INDEX_NONE)
+		{
+			Op.SetError(Errors::InvalidUser());
+			return MakeFulfilledPromise<void>().GetFuture();
+		}
+
+		TArray<FUniqueNetIdRef> UserIds;
+		UserIds.Reserve(Params.UserIds.Num());
+		for(const FOnlineAccountIdHandle& TargetAccountIdHandle : Params.UserIds)
+		{
+			FUniqueNetIdPtr TargetUserNetId = AuthInterface->GetUniqueNetId(TargetAccountIdHandle);
+			if (!TargetUserNetId)
+			{
+				Op.SetError(Errors::InvalidParams());
+				return MakeFulfilledPromise<void>().GetFuture();
+			}
+			UserIds.Emplace(TargetUserNetId.ToSharedRef());
+		}
+
+		TPromise<void> Promise;
+		TFuture<void> Future = Promise.GetFuture();
+
+		MakeMulticastAdapter(this, UserInterface->OnQueryUserInfoCompleteDelegates[LocalUserNum],
+			[WeakOp = Op.AsWeak(), Promise = MoveTemp(Promise), UserIds](int32 LocalUserNum, bool bSuccess, const TArray<FUniqueNetIdRef>& InUserIds, const FString& ErrorStr) mutable
+		{
+			const bool bShouldHandle = InUserIds == UserIds;
+			if (bShouldHandle)
+			{
+				if (TOnlineAsyncOpPtr<FQueryUserInfo> Op = WeakOp.Pin())
+				{
+					if (bSuccess)
+					{
+						Op->SetResult({});
+					}
+					else
+					{
+						Op->SetError(Errors::Unknown());
+					}
+				}
+				Promise.SetValue();
+			}
+			return bShouldHandle;
+		});
+
+		UserInterface->QueryUserInfo(LocalUserNum, UserIds);
+
+		return Future;
+	});
+
+	Op->Enqueue(GetSerialQueue());
+	return Op->GetHandle();
+}
+
+TOnlineResult<FGetUserInfo> FUserInfoOSSAdapter::GetUserInfo(FGetUserInfo::Params&& Params)
+{
+	const FAuthOSSAdapter* AuthInterface = Services.Get<FAuthOSSAdapter>();
+
+	const int32 LocalUserNum = AuthInterface->GetLocalUserNum(Params.LocalUserId);
+	if (LocalUserNum == INDEX_NONE)
+	{
+		return TOnlineResult<FGetUserInfo>(Errors::InvalidUser());
+	}
+
+	const FUniqueNetIdPtr TargetUserNetId = AuthInterface->GetUniqueNetId(Params.UserId);
+	if (!TargetUserNetId)
+	{
+		return TOnlineResult<FGetUserInfo>(Errors::InvalidParams());
+	}
+
+	const TSharedPtr<FOnlineUser> OnlineUser = UserInterface->GetUserInfo(LocalUserNum, *TargetUserNetId);
+	if (!OnlineUser)
+	{
+		return TOnlineResult<FGetUserInfo>(Errors::NotFound());
+	}
+
+	TSharedRef<FUserInfo> UserInfo = MakeShared<FUserInfo>();
+	UserInfo->UserId = Params.UserId;
+	UserInfo->DisplayName = OnlineUser->GetDisplayName();
+
+	return TOnlineResult<FGetUserInfo>({ UserInfo });
+}
+
+/* UE::Online */ }
diff --git a/Engine/Plugins/Online/Experimental/OnlineServicesOSSAdapter/Source/Public/Online/UserInfoOSSAdapter.h b/Engine/Plugins/Online/Experimental/OnlineServicesOSSAdapter/Source/Public/Online/UserInfoOSSAdapter.h
new file mode 100644
--- /dev/null
+++ b/Engine/Plugins/Online/Experimental/OnlineServicesOSSAdapter/Source/Public/Online/UserInfoOSSAdapter.h
@@ -0,0 +1,31 @@
+// Copyright Epic Games, Inc. All Rights Reserved.
+
+#pragma once
+
+#include "Online/UserInfoCommon.h"
+
+#include "OnlineSubsystemTypes.h"
+
+using IOnlineUserPtr = TSharedPtr<class IOnlineUser>;
+
+namespace UE::Online {
+
+class FUserInfoOSSAdapter : public FUserInfoCommon
+{
+public:
+	using Super = FUserInfoCommon;
+
+	using FUserInfoCommon::FUserInfoCommon;
+
+	// IOnlineComponent
+	virtual void Initialize() override;
+
+	// IUserInfo
+	virtual TOnlineAsyncOpHandle<FQueryUserInfo> QueryUserInfo(FQueryUserInfo::Params&& Params) override;
+	virtual TOnlineResult<FGetUserInfo> GetUserInfo(FGetUserInfo::Params&& Params) override;
+
+protected:
+	IOnlineUserPtr UserInterface = nullptr;
+};
+
+/* UE::Online */ }
