- Instantiate TitleFile, UserFile and UserInfo components in OSSAdapter

diff --git a/Engine/Plugins/Online/Experimental/OnlineServicesOSSAdapter/Source/Private/Online/OnlineServicesOSSAdapter.cpp b/Engine/Plugins/Online/Experimental/OnlineServicesOSSAdapter/Source/Private/Online/OnlineServicesOSSAdapter.cpp
--- a/Engine/Plugins/Online/Experimental/OnlineServicesOSSAdapter/Source/Private/Online/OnlineServicesOSSAdapter.cpp
+++ b/Engine/Plugins/Online/Experimental/OnlineServicesOSSAdapter/Source/Private/Online/OnlineServicesOSSAdapter.cpp
@@ -10,6 +10,9 @@
 #include "Online/PresenceOSSAdapter.h"
 #include "Online/PrivilegesOSSAdapter.h"
 #include "Online/StatsOSSAdapter.h"
+#include "Online/TitleFileOSSAdapter.h"
+#include "Online/UserFileOSSAdapter.h"
+#include "Online/UserInfoOSSAdapter.h"
 
 #include "OnlineSubsystem.h"
 
@@ -38,6 +41,18 @@ void FOnlineServicesOSSAdapter::RegisterComponents()
 	{
 		Components.Register<FExternalUIOSSAdapter>(*this);
 	}
+	if (Subsystem->GetTitleFileInterface().IsValid())
+	{
+		Components.Register<FTitleFileOSSAdapter>(*this);
+	}
+	if (Subsystem->GetUserCloudInterface().IsValid())
+	{
+		Components.Register<FUserFileOSSAdapter>(*this);
+	}
+	if (Subsystem->GetUserInterface().IsValid())
+	{
+		Components.Register<FUserInfoOSSAdapter>(*this);
+	}
 
 	FOnlineServicesCommon::RegisterComponents();
 }
