Fix thread-safety issue.

diff --git a/Engine/Plugins/Online/Experimental/OnlineServicesOSSAdapter/Source/Public/Online/OnlineIdOSSAdapter.h b/Engine/Plugins/Online/Experimental/OnlineServicesOSSAdapter/Source/Public/Online/OnlineIdOSSAdapter.h
--- a/Engine/Plugins/Online/Experimental/OnlineServicesOSSAdapter/Source/Public/Online/OnlineIdOSSAdapter.h
+++ b/Engine/Plugins/Online/Experimental/OnlineServicesOSSAdapter/Source/Public/Online/OnlineIdOSSAdapter.h
@@ -61,12 +61,13 @@ public:
 	// Returns a copy as it's not thread safe to return a pointer/ref to an element of an array that can be relocated by another thread.
 	FUniqueNetIdPtr GetIdValue(const FOnlineAccountIdHandle Handle) const
 	{
-		if (Handle.GetOnlineServicesType() == OnlineServicesType
-			&& Handle.IsValid()
-			&& IdValues.IsValidIndex(Handle.GetHandle() - 1))
+		if (Handle.GetOnlineServicesType() == OnlineServicesType && Handle.IsValid())
 		{
 			FReadScopeLock ReadLock(Lock);
-			return IdValues[Handle.GetHandle() - 1];
+			if (IdValues.IsValidIndex(Handle.GetHandle() - 1))
+			{
+				return IdValues[Handle.GetHandle() - 1];
+			}
 		}
 		return FUniqueNetIdPtr();
 	}
@@ -79,7 +80,8 @@ public:
 	// Begin IOnlineAccountIdRegistry
 	virtual FString ToLogString(const FOnlineAccountIdHandle& Handle) const override
 	{
-		return GetIdValue(Handle)->ToDebugString();
+		FUniqueNetIdPtr IdValue = GetIdValue(Handle);
+		return IdValue ? IdValue->ToDebugString() : FString(TEXT("invalid_id"));
 	}
 
 	virtual TArray<uint8> ToReplicationData(const FOnlineAccountIdHandle& Handle) const
