OSSv2:
- Add Leaderboards adaptor implementation;
- Stats adaptor: pass in the correct user ids when call BatchQueryStats.

diff --git a/Engine/Plugins/Online/Experimental/OnlineServicesOSSAdapter/Source/Private/Online/LeaderboardsOSSAdapter.cpp b/Engine/Plugins/Online/Experimental/OnlineServicesOSSAdapter/Source/Private/Online/LeaderboardsOSSAdapter.cpp
new file mode 100644
--- /dev/null
+++ b/Engine/Plugins/Online/Experimental/OnlineServicesOSSAdapter/Source/Private/Online/LeaderboardsOSSAdapter.cpp
@@ -0,0 +1,264 @@
+// Copyright Epic Games, Inc. All Rights Reserved.
+
+#include "Online/LeaderboardsOSSAdapter.h"
+#include "Online/AuthOSSAdapter.h"
+#include "Online/DelegateAdapter.h"
+#include "Online/ErrorsOSSAdapter.h"
+#include "Online/OnlineServicesOSSAdapter.h"
+#include "Online/OnlineIdOSSAdapter.h"
+#include "Online/OnlineErrorDefinitions.h"
+#include "OnlineSubsystem.h"
+#include "Interfaces/OnlineIdentityInterface.h"
+
+namespace UE::Online {
+
+void FLeaderboardsOSSAdapter::PostInitialize()
+{
+	Super::PostInitialize();
+
+	Auth = Services.Get<FAuthOSSAdapter>();
+	LeaderboardsInterface = static_cast<FOnlineServicesOSSAdapter&>(Services).GetSubsystem().GetLeaderboardsInterface();
+	check(LeaderboardsInterface);
+}
+
+TOnlineAsyncOpHandle<FReadEntriesForUsers> FLeaderboardsOSSAdapter::ReadEntriesForUsers(FReadEntriesForUsers::Params&& Params)
+{
+	TOnlineAsyncOpRef<FReadEntriesForUsers> Op = GetOp<FReadEntriesForUsers>(MoveTemp(Params));
+
+	Op->Then([this](TOnlineAsyncOp<FReadEntriesForUsers>& Op)
+	{
+		const FUniqueNetIdPtr LocalUniqueNetId = Auth->GetUniqueNetId(Op.GetParams().LocalUserId);
+
+		if (!LocalUniqueNetId)
+		{
+			Op.SetError(Errors::InvalidUser());
+			return MakeFulfilledPromise<void>().GetFuture();
+		}
+
+		if (Op.GetParams().UserIds.IsEmpty() || Op.GetParams().BoardName.IsEmpty())
+		{
+			Op.SetError(Errors::InvalidParams());
+			return MakeFulfilledPromise<void>().GetFuture();
+		}
+
+		TArray<FUniqueNetIdRef> NetIds;
+		for (const FOnlineAccountIdHandle& UserId : Op.GetParams().UserIds)
+		{
+			const FUniqueNetIdPtr UniqueNetId = Auth->GetUniqueNetId(UserId);
+			if (!UniqueNetId)
+			{
+				UE_LOG(LogTemp, Warning, TEXT("Failed to get unique net id for user %s"), *ToLogString(UserId));
+				Op.SetError(Errors::InvalidUser());
+				return MakeFulfilledPromise<void>().GetFuture();
+			}
+
+			NetIds.Add(UniqueNetId.ToSharedRef());
+		}
+
+		if (TOptional<FOnlineError> OnlineError = PrepareLeaderboardReadObject(Op.GetParams().BoardName))
+		{
+			Op.SetError(MoveTemp(OnlineError.GetValue()));
+			return MakeFulfilledPromise<void>().GetFuture();
+		}
+
+		TPromise<void> Promise;
+		TFuture<void> Future = Promise.GetFuture();
+
+		MakeMulticastAdapter(this, LeaderboardsInterface->OnLeaderboardReadCompleteDelegates,
+			[this, WeakOp = Op.AsWeak(), Promise = MoveTemp(Promise)](bool bSuccess) mutable
+		{
+			if (TOnlineAsyncOpPtr<FReadEntriesForUsers> Op = WeakOp.Pin())
+			{
+				if (!bSuccess)
+				{
+					ReadObject.Reset();
+					Op->SetError(Errors::Unknown());
+				}
+				else
+				{
+					FReadEntriesForUsers::Result Result;
+					ReadLeaderboardsResultFromV1ToV2(Op->GetParams().BoardName, Result.Entries);
+					Op->SetResult(MoveTemp(Result));
+				}
+			}
+			Promise.SetValue();
+		});
+
+		FOnlineLeaderboardReadRef ReadObjectRef = ReadObject.ToSharedRef();
+		if (!LeaderboardsInterface->ReadLeaderboards(NetIds, ReadObjectRef))
+		{
+			Op.SetError(Errors::Unknown());
+			return MakeFulfilledPromise<void>().GetFuture();
+		}
+
+		return Future;
+	})
+	.Enqueue(GetSerialQueue());
+
+	return Op->GetHandle();
+}
+
+TOnlineAsyncOpHandle<FReadEntriesAroundRank> FLeaderboardsOSSAdapter::ReadEntriesAroundRank(FReadEntriesAroundRank::Params&& Params)
+{
+	TOnlineAsyncOpRef<FReadEntriesAroundRank> Op = GetOp<FReadEntriesAroundRank>(MoveTemp(Params));
+
+	Op->Then([this](TOnlineAsyncOp<FReadEntriesAroundRank>& Op)
+	{
+		const FUniqueNetIdPtr LocalUserId = Auth->GetUniqueNetId(Op.GetParams().LocalUserId);
+		if (!LocalUserId)
+		{
+			Op.SetError(Errors::InvalidUser());
+			return MakeFulfilledPromise<void>().GetFuture();
+		}
+
+		if (Op.GetParams().Limit == 0 || Op.GetParams().BoardName.IsEmpty())
+		{
+			Op.SetError(Errors::InvalidParams());
+			return MakeFulfilledPromise<void>().GetFuture();
+		}
+
+		if (TOptional<FOnlineError> OnlineError = PrepareLeaderboardReadObject(Op.GetParams().BoardName))
+		{
+			Op.SetError(MoveTemp(OnlineError.GetValue()));
+			return MakeFulfilledPromise<void>().GetFuture();
+		}
+
+		TPromise<void> Promise;
+		TFuture<void> Future = Promise.GetFuture();
+
+		MakeMulticastAdapter(this, LeaderboardsInterface->OnLeaderboardReadCompleteDelegates,
+			[this, WeakOp = Op.AsWeak(), Promise = MoveTemp(Promise)](bool bSuccess) mutable
+		{
+			if (TOnlineAsyncOpPtr<FReadEntriesAroundRank> Op = WeakOp.Pin())
+			{
+				if (!bSuccess)
+				{
+					ReadObject.Reset();
+					Op->SetError(Errors::Unknown());
+				}
+				else
+				{
+					FReadEntriesAroundRank::Result Result;
+					ReadLeaderboardsResultFromV1ToV2(Op->GetParams().BoardName, Result.Entries);
+					Op->SetResult(MoveTemp(Result));
+				}
+			}
+			Promise.SetValue();
+		});
+
+		FOnlineLeaderboardReadRef ReadObjectRef = ReadObject.ToSharedRef();
+		if (!LeaderboardsInterface->ReadLeaderboardsAroundRank(Op.GetParams().Rank, Op.GetParams().Limit, ReadObjectRef))
+		{
+			Op.SetError(Errors::Unknown());
+			return MakeFulfilledPromise<void>().GetFuture();
+		}
+
+		return Future;
+	})
+	.Enqueue(GetSerialQueue());
+
+	return Op->GetHandle();
+}
+
+TOnlineAsyncOpHandle<FReadEntriesAroundUser> FLeaderboardsOSSAdapter::ReadEntriesAroundUser(FReadEntriesAroundUser::Params&& Params)
+{
+	TOnlineAsyncOpRef<FReadEntriesAroundUser> Op = GetOp<FReadEntriesAroundUser>(MoveTemp(Params));
+
+	Op->Then([this](TOnlineAsyncOp<FReadEntriesAroundUser>& Op)
+	{
+		const FUniqueNetIdPtr LocalUserId = Auth->GetUniqueNetId(Op.GetParams().LocalUserId);
+		if (!LocalUserId)
+		{
+			Op.SetError(Errors::InvalidUser());
+			return MakeFulfilledPromise<void>().GetFuture();
+		}
+
+		const FUniqueNetIdPtr TargetUniqueNetId = Auth->GetUniqueNetId(Op.GetParams().UserId);
+		if (!TargetUniqueNetId)
+		{
+			Op.SetError(Errors::InvalidUser());
+			return MakeFulfilledPromise<void>().GetFuture();
+		}
+
+		if (Op.GetParams().Limit == 0 || Op.GetParams().BoardName.IsEmpty())
+		{
+			Op.SetError(Errors::InvalidParams());
+			return MakeFulfilledPromise<void>().GetFuture();
+		}
+
+		if (TOptional<FOnlineError> OnlineError = PrepareLeaderboardReadObject(Op.GetParams().BoardName))
+		{
+			Op.SetError(MoveTemp(OnlineError.GetValue()));
+			return MakeFulfilledPromise<void>().GetFuture();
+		}
+
+		TPromise<void> Promise;
+		TFuture<void> Future = Promise.GetFuture();
+
+		MakeMulticastAdapter(this, LeaderboardsInterface->OnLeaderboardReadCompleteDelegates,
+			[this, WeakOp = Op.AsWeak(), Promise = MoveTemp(Promise)](bool bSuccess) mutable
+		{
+			if (TOnlineAsyncOpPtr<FReadEntriesAroundUser> Op = WeakOp.Pin())
+			{
+				if (!bSuccess)
+				{
+					ReadObject.Reset();
+					Op->SetError(Errors::Unknown());
+				}
+				else
+				{
+					FReadEntriesAroundUser::Result Result;
+					ReadLeaderboardsResultFromV1ToV2(Op->GetParams().BoardName, Result.Entries);
+					Op->SetResult(MoveTemp(Result));
+				}
+			}
+			Promise.SetValue();
+		});
+
+		FOnlineLeaderboardReadRef ReadObjectRef = ReadObject.ToSharedRef();
+		if (!LeaderboardsInterface->ReadLeaderboardsAroundUser(TargetUniqueNetId.ToSharedRef(), Op.GetParams().Limit, ReadObjectRef))
+		{
+			Op.SetError(Errors::Unknown());
+			return MakeFulfilledPromise<void>().GetFuture();
+		}
+
+		return Future;
+	})
+	.Enqueue(GetSerialQueue());
+
+	return Op->GetHandle();
+}
+
+TOptional<FOnlineError> FLeaderboardsOSSAdapter::PrepareLeaderboardReadObject(const FString& BoardName)
+{
+	TOptional<FOnlineError> OnlineError;
+
+	if (ReadObject)
+	{
+		OnlineError = Errors::AlreadyPending();
+	}
+	else
+	{
+		ReadObject = MakeShareable(new FOnlineLeaderboardRead());
+		ReadObject->SortedColumn = *BoardName;
+		ReadObject->LeaderboardName = *BoardName;
+	}
+
+	return OnlineError;
+}
+
+void FLeaderboardsOSSAdapter::ReadLeaderboardsResultFromV1ToV2(const FString& BoardName, TArray<FLeaderboardEntry>& OutEntries)
+{
+	for (const FOnlineStatsRow& OnlineStatsRow : ReadObject->Rows)
+	{
+		const FVariantData* VariantData = OnlineStatsRow.Columns.Find(*BoardName);
+		FLeaderboardEntry& LeaderboardEntry = OutEntries.Emplace_GetRef();
+		LeaderboardEntry.UserId = Auth->GetAccountIdHandle(OnlineStatsRow.PlayerId.ToSharedRef());
+		LeaderboardEntry.Rank = OnlineStatsRow.Rank;
+		VariantData->GetValue(LeaderboardEntry.Score);
+	}
+
+	ReadObject.Reset();
+}
+
+/* UE::Online */ }
diff --git a/Engine/Plugins/Online/Experimental/OnlineServicesOSSAdapter/Source/Private/Online/OnlineServicesOSSAdapter.cpp b/Engine/Plugins/Online/Experimental/OnlineServicesOSSAdapter/Source/Private/Online/OnlineServicesOSSAdapter.cpp
--- a/Engine/Plugins/Online/Experimental/OnlineServicesOSSAdapter/Source/Private/Online/OnlineServicesOSSAdapter.cpp
+++ b/Engine/Plugins/Online/Experimental/OnlineServicesOSSAdapter/Source/Private/Online/OnlineServicesOSSAdapter.cpp
@@ -6,6 +6,7 @@
 #include "Online/AuthOSSAdapter.h"
 #include "Online/ConnectivityOSSAdapter.h"
 #include "Online/ExternalUIOSSAdapter.h"
+#include "Online/LeaderboardsOSSAdapter.h"
 #include "Online/OnlineIdOSSAdapter.h"
 #include "Online/PresenceOSSAdapter.h"
 #include "Online/PrivilegesOSSAdapter.h"
@@ -29,6 +30,12 @@ void FOnlineServicesOSSAdapter::RegisterComponents()
 {
 	Components.Register<FAuthOSSAdapter>(*this);
 	Components.Register<FConnectivityOSSAdapter>(*this);
+
+	if (Subsystem->GetLeaderboardsInterface().IsValid())
+	{
+		Components.Register<FLeaderboardsOSSAdapter>(*this);
+	}
+
 	Components.Register<FPresenceOSSAdapter>(*this);
 	Components.Register<FPrivilegesOSSAdapter>(*this);
 	Components.Register<FStatsOSSAdapter>(*this);
diff --git a/Engine/Plugins/Online/Experimental/OnlineServicesOSSAdapter/Source/Private/Online/StatsOSSAdapter.cpp b/Engine/Plugins/Online/Experimental/OnlineServicesOSSAdapter/Source/Private/Online/StatsOSSAdapter.cpp
--- a/Engine/Plugins/Online/Experimental/OnlineServicesOSSAdapter/Source/Private/Online/StatsOSSAdapter.cpp
+++ b/Engine/Plugins/Online/Experimental/OnlineServicesOSSAdapter/Source/Private/Online/StatsOSSAdapter.cpp
@@ -210,7 +210,7 @@ TOnlineAsyncOpHandle<FBatchQueryStats> FStatsOSSAdapter::BatchQueryStats(FBatchQ
 		TArray<FUniqueNetIdRef> NetIds;
 		for (const FOnlineAccountIdHandle& TargetUserId : Op.GetParams().TargetUserIds)
 		{
-			const FUniqueNetIdPtr UniqueNetId = Auth->GetUniqueNetId(Op.GetParams().LocalUserId);
+			const FUniqueNetIdPtr UniqueNetId = Auth->GetUniqueNetId(TargetUserId);
 			if (!UniqueNetId)
 			{
 				Op.SetError(Errors::InvalidUser());
diff --git a/Engine/Plugins/Online/Experimental/OnlineServicesOSSAdapter/Source/Public/Online/LeaderboardsOSSAdapter.h b/Engine/Plugins/Online/Experimental/OnlineServicesOSSAdapter/Source/Public/Online/LeaderboardsOSSAdapter.h
new file mode 100644
--- /dev/null
+++ b/Engine/Plugins/Online/Experimental/OnlineServicesOSSAdapter/Source/Public/Online/LeaderboardsOSSAdapter.h
@@ -0,0 +1,39 @@
+// Copyright Epic Games, Inc. All Rights Reserved.
+
+#pragma once
+
+#include "Online/LeaderboardsCommon.h"
+#include "AuthOSSAdapter.h"
+#include "OnlineSubsystemTypes.h"
+#include "Interfaces/OnlineLeaderboardInterface.h"
+
+class IOnlineLeaderboards;
+using IOnlineLeaderboardsPtr = TSharedPtr<IOnlineLeaderboards>;
+
+namespace UE::Online {
+
+class FLeaderboardsOSSAdapter : public FLeaderboardsCommon
+{
+public:
+	using Super = FLeaderboardsCommon;
+
+	using FLeaderboardsCommon::FLeaderboardsCommon;
+
+	// IOnlineComponent
+	virtual void PostInitialize() override;
+
+	// ILeaderboards
+	virtual TOnlineAsyncOpHandle<FReadEntriesForUsers> ReadEntriesForUsers(FReadEntriesForUsers::Params&& Params) override;
+	virtual TOnlineAsyncOpHandle<FReadEntriesAroundRank> ReadEntriesAroundRank(FReadEntriesAroundRank::Params&& Params) override;
+	virtual TOnlineAsyncOpHandle<FReadEntriesAroundUser> ReadEntriesAroundUser(FReadEntriesAroundUser::Params&& Params) override;
+
+protected:
+	TOptional<FOnlineError> PrepareLeaderboardReadObject(const FString& BoardName);
+	void ReadLeaderboardsResultFromV1ToV2(const FString& BoardName, TArray<FLeaderboardEntry>& OutEntries);
+
+	const FAuthOSSAdapter* Auth;
+	IOnlineLeaderboardsPtr LeaderboardsInterface;
+	FOnlineLeaderboardReadPtr ReadObject;
+};
+
+/* UE::Online */ }
