OSSv2 Sessions Interface: Complete OSS Adapter Implementation
- Added StartMatchmaking parameters and state check common methods
- Refactor of TOnlineBasicIdRegistry to take the Online Services type as constructor parameter instead of template
- Refactor of FOnlineNetIdRegistry to serve as base class for Account and Session registries in OSS Adapter

diff --git a/Engine/Plugins/Online/Experimental/OnlineServices/Source/OnlineServicesCommon/Private/Online/SessionsCommon.cpp b/Engine/Plugins/Online/Experimental/OnlineServices/Source/OnlineServicesCommon/Private/Online/SessionsCommon.cpp
--- a/Engine/Plugins/Online/Experimental/OnlineServices/Source/OnlineServicesCommon/Private/Online/SessionsCommon.cpp
+++ b/Engine/Plugins/Online/Experimental/OnlineServices/Source/OnlineServicesCommon/Private/Online/SessionsCommon.cpp
@@ -613,6 +613,50 @@ namespace UE::Online {
 		return Errors::Success();
 	}
 
+	FOnlineError FSessionsCommon::CheckStartMatchmakingParams(const FStartMatchmaking::Params& Params)
+	{
+		if (Params.SessionSettings.NumMaxPrivateConnections == 0 && Params.SessionSettings.NumMaxPublicConnections == 0)
+		{
+			UE_LOG(LogTemp, Warning, TEXT("[FSessionsCommon::CheckStartMatchmakingParams] Could not create session with invalid NumMaxPrivateConnections and NumMaxPublicConnections"));
+
+			return Errors::InvalidParams();
+		}
+
+		return Errors::Success();
+	}
+
+	FOnlineError FSessionsCommon::CheckStartMatchmakingState(const FStartMatchmaking::Params& Params)
+	{
+		// Check if a session with that name already exists
+		if (TSharedRef<FSession>* FoundSession = LocalSessionsByName.Find(Params.SessionName))
+		{
+			UE_LOG(LogTemp, Warning, TEXT("[FSessionsCommon::CheckStartMatchmakingState] Could not join session with name [%s]. A session with that name already exists"), *Params.SessionName.ToString());
+
+			return Errors::InvalidState(); // TODO: New error: Session with name %s already exists
+		}
+
+		// User login check for all local users
+		IAuthPtr Auth = Services.GetAuthInterface();
+
+		TArray<FOnlineAccountIdHandle> LocalUserIds;
+		LocalUserIds.Reserve(Params.LocalUsers.Num());
+		Params.LocalUsers.GenerateKeyArray(LocalUserIds);
+
+		for (const FOnlineAccountIdHandle& LocalUserId : LocalUserIds)
+		{
+			if (!Auth->IsLoggedIn(LocalUserId))
+			{
+				UE_LOG(LogTemp, Warning, TEXT("[FSessionsCommon::CheckStartMatchmakingState] Could not join session with user [%s] not logged in"), *ToLogString(LocalUserId));
+
+				return Errors::InvalidUser();
+			}
+		}
+
+		// TODO: Check that only one session has bUsesPresence set
+
+		return Errors::Success();
+	}
+
 	FOnlineError FSessionsCommon::CheckJoinSessionParams(const FJoinSession::Params& Params)
 	{
 		if (!Params.LocalUserId.IsValid())
diff --git a/Engine/Plugins/Online/Experimental/OnlineServices/Source/OnlineServicesCommon/Public/Online/OnlineIdCommon.h b/Engine/Plugins/Online/Experimental/OnlineServices/Source/OnlineServicesCommon/Public/Online/OnlineIdCommon.h
--- a/Engine/Plugins/Online/Experimental/OnlineServices/Source/OnlineServicesCommon/Public/Online/OnlineIdCommon.h
+++ b/Engine/Plugins/Online/Experimental/OnlineServices/Source/OnlineServicesCommon/Public/Online/OnlineIdCommon.h
@@ -9,15 +9,24 @@ namespace UE::Online {
 /**
  * A net id registry suitable for use with trivial immutable keys
  */
-template<typename IdType, typename IdValueType, EOnlineServices OnlineServicesType>
+template<typename IdType, typename IdValueType>
 class TOnlineBasicIdRegistry
 {
 public:
+	TOnlineBasicIdRegistry(EOnlineServices InOnlineServicesType)
+		: OnlineServicesType(InOnlineServicesType)
+	{
+
+	}
+
 	virtual ~TOnlineBasicIdRegistry() = default;
 
 	using HandleType = TOnlineIdHandle<IdType>;
 
-	const HandleType InvalidHandle = HandleType(OnlineServicesType, 0);
+	HandleType GetInvalidHandle() const
+	{
+		return HandleType(OnlineServicesType, 0);
+	}
 
 	HandleType FindHandle(const IdValueType& IdValue) const
 	{
@@ -26,7 +35,7 @@ public:
 		{
 			return *FoundHandle;
 		}
-		return InvalidHandle;
+		return GetInvalidHandle();
 	}
 
 	HandleType FindOrAddHandle(const IdValueType& IdValue)
@@ -65,26 +74,28 @@ public:
 		return IdValueType();
 	}
 
-	static inline bool ValidateOnlineId(const HandleType& Handle)
+	inline bool ValidateOnlineId(const HandleType& Handle) const
 	{
 		return ensure(Handle.GetOnlineServicesType() == OnlineServicesType) && Handle.IsValid();
 	}
 
 private:
+	EOnlineServices OnlineServicesType;
+
 	mutable FRWLock Lock;
 
 	TArray<IdValueType> IdValues;
 	TMap<IdValueType, HandleType> IdValueToHandleMap;
 };
 
-template<typename IdValueType, EOnlineServices OnlineServicesType>
-using TOnlineBasicAccountIdRegistry = TOnlineBasicIdRegistry<OnlineIdHandleTags::FAccount, IdValueType, OnlineServicesType>;
+template<typename IdValueType>
+using TOnlineBasicAccountIdRegistry = TOnlineBasicIdRegistry<OnlineIdHandleTags::FAccount, IdValueType>;
 
-template<typename IdValueType, EOnlineServices OnlineServicesType>
-using TOnlineBasicSessionIdRegistry = TOnlineBasicIdRegistry<OnlineIdHandleTags::FSession, IdValueType, OnlineServicesType>;
+template<typename IdValueType>
+using TOnlineBasicSessionIdRegistry = TOnlineBasicIdRegistry<OnlineIdHandleTags::FSession, IdValueType>;
 
-template<typename IdValueType, EOnlineServices OnlineServicesType>
-using TOnlineBasicSessionInviteIdRegistry = TOnlineBasicIdRegistry<OnlineIdHandleTags::FSessionInvite, IdValueType, OnlineServicesType>;
+template<typename IdValueType>
+using TOnlineBasicSessionInviteIdRegistry = TOnlineBasicIdRegistry<OnlineIdHandleTags::FSessionInvite, IdValueType>;
 
 /* UE::Online */ }
 
diff --git a/Engine/Plugins/Online/Experimental/OnlineServices/Source/OnlineServicesCommon/Public/Online/SessionsCommon.h b/Engine/Plugins/Online/Experimental/OnlineServices/Source/OnlineServicesCommon/Public/Online/SessionsCommon.h
--- a/Engine/Plugins/Online/Experimental/OnlineServices/Source/OnlineServicesCommon/Public/Online/SessionsCommon.h
+++ b/Engine/Plugins/Online/Experimental/OnlineServices/Source/OnlineServicesCommon/Public/Online/SessionsCommon.h
@@ -12,8 +12,7 @@ class FOnlineServicesCommon;
 
 static FName CONNECT_STRING_TAG = TEXT("CONNECT_STRING");
 
-template<EOnlineServices OnlineServicesType>
-class TOnlineSessionIdStringRegistry : public IOnlineSessionIdRegistry
+class FOnlineSessionIdStringRegistry : public IOnlineSessionIdRegistry
 {
 public:
 	// Begin IOnlineSessionIdRegistry
@@ -56,14 +55,24 @@ public:
 	}
 	// End IOnlineSessionIdRegistry
 
-	virtual ~TOnlineSessionIdStringRegistry() = default;
+	inline bool IsSessionIdExpired(const FOnlineSessionIdHandle& InHandle) const
+	{
+		return BasicRegistry.FindIdValue(InHandle).IsEmpty();
+	}
+
+	FOnlineSessionIdStringRegistry(EOnlineServices OnlineServicesType)
+		: BasicRegistry(OnlineServicesType)
+	{
+
+	}
+
+	virtual ~FOnlineSessionIdStringRegistry() = default;
 
 public:
-	TOnlineBasicSessionIdRegistry<FString, OnlineServicesType> BasicRegistry;
+	TOnlineBasicSessionIdRegistry<FString> BasicRegistry;
 };
 
-template<EOnlineServices OnlineServicesType>
-class TOnlineSessionInviteIdStringRegistry : public IOnlineSessionInviteIdRegistry
+class FOnlineSessionInviteIdStringRegistry : public IOnlineSessionInviteIdRegistry
 {
 public:
 	// Begin IOnlineSessionIdRegistry
@@ -106,10 +115,16 @@ public:
 	}
 	// End IOnlineSessionIdRegistry
 
-	virtual ~TOnlineSessionInviteIdStringRegistry() = default;
+	FOnlineSessionInviteIdStringRegistry(EOnlineServices OnlineServicesType)
+		: BasicRegistry(OnlineServicesType)
+	{
+
+	}
+
+	virtual ~FOnlineSessionInviteIdStringRegistry() = default;
 
 public:
-	TOnlineBasicSessionInviteIdRegistry<FString, OnlineServicesType> BasicRegistry;
+	TOnlineBasicSessionInviteIdRegistry<FString> BasicRegistry;
 };
 
 class ONLINESERVICESCOMMON_API FSessionsCommon : public TOnlineComponent<ISessions>
@@ -162,6 +177,9 @@ protected:
 	FOnlineError CheckFindSessionsParams(const FFindSessions::Params& Params);
 	FOnlineError CheckFindSessionsState(const FFindSessions::Params& Params);
 
+	FOnlineError CheckStartMatchmakingParams(const FStartMatchmaking::Params& Params);
+	FOnlineError CheckStartMatchmakingState(const FStartMatchmaking::Params& Params);
+
 	FOnlineError CheckJoinSessionParams(const FJoinSession::Params& Params);
 	FOnlineError CheckJoinSessionState(const FJoinSession::Params& Params);
 
diff --git a/Engine/Plugins/Online/Experimental/OnlineServices/Source/OnlineServicesCommonEngineUtils/Private/Online/SessionsLAN.cpp b/Engine/Plugins/Online/Experimental/OnlineServices/Source/OnlineServicesCommonEngineUtils/Private/Online/SessionsLAN.cpp
--- a/Engine/Plugins/Online/Experimental/OnlineServices/Source/OnlineServicesCommonEngineUtils/Private/Online/SessionsLAN.cpp
+++ b/Engine/Plugins/Online/Experimental/OnlineServices/Source/OnlineServicesCommonEngineUtils/Private/Online/SessionsLAN.cpp
@@ -11,21 +11,22 @@
 namespace UE::Online {
 
 /** FOnlineSessionIdRegistryLAN */
-
-FOnlineSessionIdRegistryLAN& FOnlineSessionIdRegistryLAN::Get()
+FOnlineSessionIdRegistryLAN::FOnlineSessionIdRegistryLAN(EOnlineServices ServicesType)
+	: FOnlineSessionIdStringRegistry(ServicesType)
 {
-	static FOnlineSessionIdRegistryLAN Instance;
-	return Instance;
+
 }
 
-FOnlineSessionIdHandle FOnlineSessionIdRegistryLAN::GetNextSessionId()
+FOnlineSessionIdRegistryLAN& FOnlineSessionIdRegistryLAN::GetChecked(EOnlineServices ServicesType)
 {
-	return BasicRegistry.FindOrAddHandle(FGuid::NewGuid().ToString());
+	FOnlineSessionIdRegistryLAN* SessionIdRegistry = static_cast<FOnlineSessionIdRegistryLAN*>(FOnlineIdRegistryRegistry::Get().GetSessionIdRegistry(ServicesType));
+	check(SessionIdRegistry);
+	return *SessionIdRegistry;
 }
 
-bool FOnlineSessionIdRegistryLAN::IsSessionIdExpired(const FOnlineSessionIdHandle& InHandle) const
+FOnlineSessionIdHandle FOnlineSessionIdRegistryLAN::GetNextSessionId()
 {
-	return BasicRegistry.FindIdValue(InHandle).IsEmpty();
+	return BasicRegistry.FindOrAddHandle(FGuid::NewGuid().ToString());
 }
 
 /** FSessionLAN */
@@ -128,7 +129,7 @@ TOnlineAsyncOpHandle<FCreateSession> FSessionsLAN::CreateSession(FCreateSession:
 		TSharedRef<FSessionLAN> NewSessionLANRef = MakeShared<FSessionLAN>();
 		NewSessionLANRef->CurrentState = ESessionState::Valid;
 		NewSessionLANRef->OwnerUserId = OpParams.LocalUserId;
-		NewSessionLANRef->SessionId = FOnlineSessionIdRegistryLAN::Get().GetNextSessionId();
+		NewSessionLANRef->SessionId = FOnlineSessionIdRegistryLAN::GetChecked(Services.GetServicesProvider()).GetNextSessionId();
 		NewSessionLANRef->SessionSettings = OpParams.SessionSettings;
 
 		// For LAN sessions, we'll add all the Session members manually instead of calling JoinSession since there is no API calls involved
@@ -256,7 +257,7 @@ TOnlineAsyncOpHandle<FJoinSession> FSessionsLAN::JoinSession(FJoinSession::Param
 		if (GetSessionByIdResult.IsError())
 		{
 			// If no result is found, the id might be expired, which we should notify
-			if (FOnlineSessionIdRegistryLAN::Get().IsSessionIdExpired(OpParams.SessionId))
+			if (FOnlineSessionIdRegistryLAN::GetChecked(Services.GetServicesProvider()).IsSessionIdExpired(OpParams.SessionId))
 			{
 				UE_LOG(LogTemp, Warning, TEXT("[FSessionsLAN::JoinSession] SessionId parameter [%s] is expired. Please call FindSessions to get an updated list of available sessions "), *ToLogString(OpParams.SessionId));
 			}
diff --git a/Engine/Plugins/Online/Experimental/OnlineServices/Source/OnlineServicesCommonEngineUtils/Public/Online/SessionsLAN.h b/Engine/Plugins/Online/Experimental/OnlineServices/Source/OnlineServicesCommonEngineUtils/Public/Online/SessionsLAN.h
--- a/Engine/Plugins/Online/Experimental/OnlineServices/Source/OnlineServicesCommonEngineUtils/Public/Online/SessionsLAN.h
+++ b/Engine/Plugins/Online/Experimental/OnlineServices/Source/OnlineServicesCommonEngineUtils/Public/Online/SessionsLAN.h
@@ -10,14 +10,15 @@ namespace UE::Online {
 
 class FOnlineServicesCommon;
 
-class ONLINESERVICESCOMMONENGINEUTILS_API FOnlineSessionIdRegistryLAN : public TOnlineSessionIdStringRegistry<EOnlineServices::Default>
+class ONLINESERVICESCOMMONENGINEUTILS_API FOnlineSessionIdRegistryLAN : public FOnlineSessionIdStringRegistry
 {
 public:
-	static FOnlineSessionIdRegistryLAN& Get();
+	static FOnlineSessionIdRegistryLAN& GetChecked(EOnlineServices ServicesType);
 
 	FOnlineSessionIdHandle GetNextSessionId();
 
-	bool IsSessionIdExpired(const FOnlineSessionIdHandle& InHandle) const;
+protected:
+	FOnlineSessionIdRegistryLAN(EOnlineServices ServicesType);
 };
 
 class ONLINESERVICESCOMMONENGINEUTILS_API FSessionLAN : public FSession
diff --git a/Engine/Plugins/Online/Experimental/OnlineServicesEOSGS/Source/Private/Online/OnlineIdEOSGS.cpp b/Engine/Plugins/Online/Experimental/OnlineServicesEOSGS/Source/Private/Online/OnlineIdEOSGS.cpp
--- a/Engine/Plugins/Online/Experimental/OnlineServicesEOSGS/Source/Private/Online/OnlineIdEOSGS.cpp
+++ b/Engine/Plugins/Online/Experimental/OnlineServicesEOSGS/Source/Private/Online/OnlineIdEOSGS.cpp
@@ -12,6 +12,12 @@ namespace UE::Online {
 const uint8 OnlineIdEOSUtf8BufferLength = 32;
 const uint8 OnlineIdEOSHexBufferLength = 16;
 
+FOnlineAccountIdRegistryEOSGS::FOnlineAccountIdRegistryEOSGS()
+	: Registry(EOnlineServices::Epic)
+{
+
+}
+
 FOnlineAccountIdHandle FOnlineAccountIdRegistryEOSGS::FindAccountId(const EOS_ProductUserId ProductUserId) const
 {
 	return Registry.FindHandle(ProductUserId);
@@ -62,7 +68,7 @@ FOnlineAccountIdHandle FOnlineAccountIdRegistryEOSGS::FromReplicationData(const
 		EOS_ProductUserId ProductUserId = EOS_ProductUserId_FromString(EosBuffer);
 		return FindOrAddAccountId(ProductUserId);
 	}
-	return Registry.InvalidHandle;
+	return Registry.GetInvalidHandle();
 }
 
 IOnlineAccountIdRegistryEOSGS& FOnlineAccountIdRegistryEOSGS::GetRegistered()
@@ -84,7 +90,7 @@ FOnlineAccountIdHandle FOnlineAccountIdRegistryEOSGS::FindOrAddAccountId(const E
 	{
 		return Registry.FindOrAddHandle(ProductUserId);
 	}
-	return Registry.InvalidHandle;
+	return Registry.GetInvalidHandle();
 }
 
 EOS_ProductUserId GetProductUserId(const FOnlineAccountIdHandle& Handle)
diff --git a/Engine/Plugins/Online/Experimental/OnlineServicesEOSGS/Source/Private/Online/SessionsEOSGS.cpp b/Engine/Plugins/Online/Experimental/OnlineServicesEOSGS/Source/Private/Online/SessionsEOSGS.cpp
--- a/Engine/Plugins/Online/Experimental/OnlineServicesEOSGS/Source/Private/Online/SessionsEOSGS.cpp
+++ b/Engine/Plugins/Online/Experimental/OnlineServicesEOSGS/Source/Private/Online/SessionsEOSGS.cpp
@@ -7,6 +7,7 @@
 #include "Online/OnlineServicesEOSGS.h"
 #include "Online/OnlineServicesEOSGSTypes.h"
 #include "Online/NboSerializerEOSGSSvc.h"
+#include "Online/OnlineErrorEOSGS.h"
 
 #include "eos_sessions.h"
 #include "eos_sessions_types.h"
@@ -30,30 +31,30 @@ EOS_EOnlineSessionPermissionLevel GetEOSGSType(const ESessionJoinPolicy& Value)
 
 /** FOnlineSessionIdRegistryEOSGS */
 
+FOnlineSessionIdRegistryEOSGS::FOnlineSessionIdRegistryEOSGS()
+	: FOnlineSessionIdRegistryLAN(EOnlineServices::Epic)
+{
+}
+
 FOnlineSessionIdRegistryEOSGS& FOnlineSessionIdRegistryEOSGS::Get()
 {
 	static FOnlineSessionIdRegistryEOSGS Instance;
 	return Instance;
 }
 
-bool FOnlineSessionIdRegistryEOSGS::IsSessionIdExpired(const FOnlineSessionIdHandle& InHandle) const
+/** FOnlineSessionInviteIdRegistryEOSGS */
+
+FOnlineSessionInviteIdRegistryEOSGS::FOnlineSessionInviteIdRegistryEOSGS()
+	: FOnlineSessionInviteIdStringRegistry(EOnlineServices::Epic)
 {
-	return BasicRegistry.FindIdValue(InHandle).IsEmpty();
 }
 
-/** FOnlineSessionInviteIdRegistryEOSGS */
-
 FOnlineSessionInviteIdRegistryEOSGS& FOnlineSessionInviteIdRegistryEOSGS::Get()
 {
 	static FOnlineSessionInviteIdRegistryEOSGS Instance;
 	return Instance;
 }
 
-bool FOnlineSessionInviteIdRegistryEOSGS::IsSessionInviteIdExpired(const FOnlineSessionInviteIdHandle& InHandle) const
-{
-	return BasicRegistry.FindIdValue(InHandle).IsEmpty();
-}
-
 /** FSessionEOSGS */
 
 FSessionEOSGS::FSessionEOSGS(const EOS_HSessionDetails& InSessionDetailsHandle)
@@ -955,6 +956,12 @@ TOnlineAsyncOpHandle<FJoinSession> FSessionsEOSGS::JoinSession(FJoinSession::Par
 	TOnlineResult<FGetSessionById> GetSessionByIdResult = GetSessionById({ OpParams.LocalUserId, OpParams.SessionId });
 	if (GetSessionByIdResult.IsError())
 	{
+		// If no result is found, the id might be expired, which we should notify
+		if (FOnlineSessionIdRegistryEOSGS::Get().IsSessionIdExpired(OpParams.SessionId))
+		{
+			UE_LOG(LogTemp, Warning, TEXT("[FSessionsEOSGS::JoinSession] SessionId parameter [%s] is expired. Please call FindSessions to get an updated list of available sessions "), *ToLogString(OpParams.SessionId));
+		}
+
 		Op->SetError(MoveTemp(GetSessionByIdResult.GetErrorValue()));
 		return Op->GetHandle();
 	}
@@ -1002,6 +1009,12 @@ TOnlineAsyncOpHandle<FJoinSession> FSessionsEOSGS::JoinSession(FJoinSession::Par
 		TOnlineResult<FGetSessionById> GetSessionByIdResult = GetSessionById({ OpParams.LocalUserId, OpParams.SessionId });
 		if (GetSessionByIdResult.IsError())
 		{
+			// If no result is found, the id might be expired, which we should notify
+			if (FOnlineSessionIdRegistryEOSGS::Get().IsSessionIdExpired(OpParams.SessionId))
+			{
+				UE_LOG(LogTemp, Warning, TEXT("[FSessionsEOSGS::JoinSession] SessionId parameter [%s] is expired. Please call FindSessions to get an updated list of available sessions "), *ToLogString(OpParams.SessionId));
+			}
+
 			Op.SetError(MoveTemp(GetSessionByIdResult.GetErrorValue()));
 			Promise.EmplaceValue();
 			return;
diff --git a/Engine/Plugins/Online/Experimental/OnlineServicesEOSGS/Source/Private/OnlineServicesEOSGSModule.cpp b/Engine/Plugins/Online/Experimental/OnlineServicesEOSGS/Source/Private/OnlineServicesEOSGSModule.cpp
--- a/Engine/Plugins/Online/Experimental/OnlineServicesEOSGS/Source/Private/OnlineServicesEOSGSModule.cpp
+++ b/Engine/Plugins/Online/Experimental/OnlineServicesEOSGS/Source/Private/OnlineServicesEOSGSModule.cpp
@@ -6,6 +6,7 @@
 #include "Online/OnlineServicesEOSGS.h"
 #include "Online/OnlineServicesRegistry.h"
 #include "Online/OnlineServicesEOSGSPlatformFactory.h"
+#include "Online/SessionsEOSGS.h"
 
 #include "CoreMinimal.h"
 
@@ -35,6 +36,8 @@ void FOnlineServicesEOSGSModule::StartupModule()
 {
 	FOnlineServicesRegistry::Get().RegisterServicesFactory(EOnlineServices::Epic, MakeUnique<FOnlineServicesFactoryEOSGS>());
 	FOnlineIdRegistryRegistry::Get().RegisterAccountIdRegistry(EOnlineServices::Epic, &FOnlineAccountIdRegistryEOSGS::Get());
+	FOnlineIdRegistryRegistry::Get().RegisterSessionIdRegistry(EOnlineServices::Epic, &FOnlineSessionIdRegistryEOSGS::Get());
+	FOnlineIdRegistryRegistry::Get().RegisterSessionInviteIdRegistry(EOnlineServices::Epic, &FOnlineSessionInviteIdRegistryEOSGS::Get());
 
 	// Initialize the platform factory on startup.  This is necessary for the SDK to bind to rendering and input very early.
 	FOnlineServicesEOSGSPlatformFactory::Get();
diff --git a/Engine/Plugins/Online/Experimental/OnlineServicesEOSGS/Source/Public/Online/OnlineIdEOSGS.h b/Engine/Plugins/Online/Experimental/OnlineServicesEOSGS/Source/Public/Online/OnlineIdEOSGS.h
--- a/Engine/Plugins/Online/Experimental/OnlineServicesEOSGS/Source/Public/Online/OnlineIdEOSGS.h
+++ b/Engine/Plugins/Online/Experimental/OnlineServicesEOSGS/Source/Public/Online/OnlineIdEOSGS.h
@@ -26,6 +26,7 @@ class FOnlineAccountIdRegistryEOSGS
 	: public IOnlineAccountIdRegistryEOSGS
 {
 public:
+	FOnlineAccountIdRegistryEOSGS();
 	virtual ~FOnlineAccountIdRegistryEOSGS() = default;
 
 	// Begin IOnlineAccountIdRegistryEOSGS
@@ -48,7 +49,7 @@ private:
 	static FOnlineAccountIdRegistryEOSGS& Get();
 	FOnlineAccountIdHandle FindOrAddAccountId(const EOS_ProductUserId ProductUserId);
 
-	TOnlineBasicAccountIdRegistry<EOS_ProductUserId, EOnlineServices::Epic> Registry;
+	TOnlineBasicAccountIdRegistry<EOS_ProductUserId> Registry;
 };
 
 EOS_ProductUserId ONLINESERVICESEOSGS_API GetProductUserId(const FOnlineAccountIdHandle& Handle);
diff --git a/Engine/Plugins/Online/Experimental/OnlineServicesEOSGS/Source/Public/Online/SessionsEOSGS.h b/Engine/Plugins/Online/Experimental/OnlineServicesEOSGS/Source/Public/Online/SessionsEOSGS.h
--- a/Engine/Plugins/Online/Experimental/OnlineServicesEOSGS/Source/Public/Online/SessionsEOSGS.h
+++ b/Engine/Plugins/Online/Experimental/OnlineServicesEOSGS/Source/Public/Online/SessionsEOSGS.h
@@ -16,20 +16,22 @@ namespace UE::Online {
 
 class FOnlineServicesEOSGS;
 
-class FOnlineSessionIdRegistryEOSGS : public TOnlineSessionIdStringRegistry<EOnlineServices::Epic>
+class FOnlineSessionIdRegistryEOSGS : public FOnlineSessionIdRegistryLAN
 {
 public:
 	static FOnlineSessionIdRegistryEOSGS& Get();
 
-	bool IsSessionIdExpired(const FOnlineSessionIdHandle& InHandle) const;
+private:
+	FOnlineSessionIdRegistryEOSGS();
 };
 
-class FOnlineSessionInviteIdRegistryEOSGS : public TOnlineSessionInviteIdStringRegistry<EOnlineServices::Epic>
+class FOnlineSessionInviteIdRegistryEOSGS : public FOnlineSessionInviteIdStringRegistry
 {
 public:
 	static FOnlineSessionInviteIdRegistryEOSGS& Get();
 
-	bool IsSessionInviteIdExpired(const FOnlineSessionInviteIdHandle& InHandle) const;
+private:
+	FOnlineSessionInviteIdRegistryEOSGS();
 };
 
 static FName EOS_SESSIONS_BUCKET_ID = TEXT("EOS_SESSIONS_BUCKET_ID");
diff --git a/Engine/Plugins/Online/Experimental/OnlineServicesNull/Source/Private/Online/OnlineServicesNullModule.cpp b/Engine/Plugins/Online/Experimental/OnlineServicesNull/Source/Private/Online/OnlineServicesNullModule.cpp
--- a/Engine/Plugins/Online/Experimental/OnlineServicesNull/Source/Private/Online/OnlineServicesNullModule.cpp
+++ b/Engine/Plugins/Online/Experimental/OnlineServicesNull/Source/Private/Online/OnlineServicesNullModule.cpp
@@ -7,6 +7,7 @@
 #include "Online/OnlineServicesRegistry.h"
 #include "Online/OnlineServicesNull.h"
 #include "Online/AuthNull.h"
+#include "Online/SessionsNull.h"
 
 namespace UE::Online
 {
@@ -33,6 +34,7 @@ void FOnlineServicesNullModule::StartupModule()
 {
 	FOnlineServicesRegistry::Get().RegisterServicesFactory(EOnlineServices::Null, MakeUnique<FOnlineServicesFactoryNull>());
 	FOnlineIdRegistryRegistry::Get().RegisterAccountIdRegistry(EOnlineServices::Null, &FOnlineAccountIdRegistryNull::Get());
+	FOnlineIdRegistryRegistry::Get().RegisterSessionIdRegistry(EOnlineServices::Null, &FOnlineSessionIdRegistryNull::Get());
 }
 
 /* UE::Online */ }
diff --git a/Engine/Plugins/Online/Experimental/OnlineServicesNull/Source/Private/Online/SessionsNull.cpp b/Engine/Plugins/Online/Experimental/OnlineServicesNull/Source/Private/Online/SessionsNull.cpp
--- a/Engine/Plugins/Online/Experimental/OnlineServicesNull/Source/Private/Online/SessionsNull.cpp
+++ b/Engine/Plugins/Online/Experimental/OnlineServicesNull/Source/Private/Online/SessionsNull.cpp
@@ -7,6 +7,19 @@
 
 namespace UE::Online {
 
+/** FOnlineSessionIdRegistryNull */
+
+FOnlineSessionIdRegistryNull::FOnlineSessionIdRegistryNull()
+	: FOnlineSessionIdRegistryLAN(EOnlineServices::Null)
+{
+}
+
+FOnlineSessionIdRegistryNull& FOnlineSessionIdRegistryNull::Get()
+{
+	static FOnlineSessionIdRegistryNull Instance;
+	return Instance;
+}
+
 /** FSessionsNull */
 
 FSessionsNull::FSessionsNull(FOnlineServicesNull& InServices)
diff --git a/Engine/Plugins/Online/Experimental/OnlineServicesNull/Source/Public/Online/NboSerializerNullSvc.h b/Engine/Plugins/Online/Experimental/OnlineServicesNull/Source/Public/Online/NboSerializerNullSvc.h
--- a/Engine/Plugins/Online/Experimental/OnlineServicesNull/Source/Public/Online/NboSerializerNullSvc.h
+++ b/Engine/Plugins/Online/Experimental/OnlineServicesNull/Source/Public/Online/NboSerializerNullSvc.h
@@ -6,7 +6,7 @@
 
 #include "Online/AuthNull.h"
 #include "Online/LobbiesNull.h"
-#include "Online/SessionsLAN.h"
+#include "Online/SessionsNull.h"
 #include "Online/CoreOnline.h"
 
 /**
@@ -27,7 +27,7 @@ inline void SerializeToBuffer(FNboSerializeToBuffer& Ar, const FOnlineAccountIdH
 
 inline void SerializeToBuffer(FNboSerializeToBuffer& Ar, const FOnlineSessionIdHandle& SessionId)
 {
-	TArray<uint8> Data = FOnlineSessionIdRegistryLAN::Get().ToReplicationData(SessionId);
+	TArray<uint8> Data = FOnlineSessionIdRegistryNull::Get().ToReplicationData(SessionId);
 	Ar << Data.Num();
 	Ar.WriteBinary(Data.GetData(), Data.Num());
 }
@@ -61,7 +61,7 @@ inline void SerializeFromBuffer(FNboSerializeFromBuffer& Ar, FOnlineSessionIdHan
 	Ar >> Size;
 	Ar.ReadBinaryArray(Data, Size);
 
-	SessionId = FOnlineSessionIdRegistryLAN::Get().FromReplicationData(Data);
+	SessionId = FOnlineSessionIdRegistryNull::Get().FromReplicationData(Data);
 }
 
 inline void SerializeFromBuffer(FNboSerializeFromBuffer& Ar, TMap<FLobbyAttributeId, FLobbyVariant>& Map)
diff --git a/Engine/Plugins/Online/Experimental/OnlineServicesNull/Source/Public/Online/SessionsNull.h b/Engine/Plugins/Online/Experimental/OnlineServicesNull/Source/Public/Online/SessionsNull.h
--- a/Engine/Plugins/Online/Experimental/OnlineServicesNull/Source/Public/Online/SessionsNull.h
+++ b/Engine/Plugins/Online/Experimental/OnlineServicesNull/Source/Public/Online/SessionsNull.h
@@ -7,6 +7,15 @@ namespace UE::Online {
 
 class FOnlineServicesNull;
 
+class FOnlineSessionIdRegistryNull : public FOnlineSessionIdRegistryLAN
+{
+public:
+	static FOnlineSessionIdRegistryNull& Get();
+
+private:
+	FOnlineSessionIdRegistryNull();
+};
+
 typedef FSessionLAN FSessionNull;
 
 class FSessionsNull : public FSessionsLAN
diff --git a/Engine/Plugins/Online/Experimental/OnlineServicesOSSAdapter/Source/Private/Online/OnlineServicesOSSAdapter.cpp b/Engine/Plugins/Online/Experimental/OnlineServicesOSSAdapter/Source/Private/Online/OnlineServicesOSSAdapter.cpp
--- a/Engine/Plugins/Online/Experimental/OnlineServicesOSSAdapter/Source/Private/Online/OnlineServicesOSSAdapter.cpp
+++ b/Engine/Plugins/Online/Experimental/OnlineServicesOSSAdapter/Source/Private/Online/OnlineServicesOSSAdapter.cpp
@@ -10,6 +10,7 @@
 #include "Online/OnlineIdOSSAdapter.h"
 #include "Online/PresenceOSSAdapter.h"
 #include "Online/PrivilegesOSSAdapter.h"
+#include "Online/SessionsOSSAdapter.h"
 #include "Online/StatsOSSAdapter.h"
 #include "Online/TitleFileOSSAdapter.h"
 #include "Online/UserFileOSSAdapter.h"
@@ -40,6 +41,10 @@ void FOnlineServicesOSSAdapter::RegisterComponents()
 	Components.Register<FPrivilegesOSSAdapter>(*this);
 	Components.Register<FStatsOSSAdapter>(*this);
 
+	if (Subsystem->GetSessionInterface().IsValid())
+	{
+		Components.Register<FSessionsOSSAdapter>(*this);
+	}
 	if (Subsystem->GetAchievementsInterface().IsValid())
 	{
 		Components.Register<FAchievementsOSSAdapter>(*this);
@@ -66,9 +71,17 @@ void FOnlineServicesOSSAdapter::RegisterComponents()
 
 void FOnlineServicesOSSAdapter::Initialize()
 {
-	AccountIdRegistry = static_cast<FOnlineUniqueNetIdRegistry*>(FOnlineIdRegistryRegistry::Get().GetAccountIdRegistry(ServicesType));
+	AccountIdRegistry = static_cast<FOnlineAccountIdRegistryOSSAdapter*>(FOnlineIdRegistryRegistry::Get().GetAccountIdRegistry(ServicesType));
 
 	FOnlineServicesCommon::Initialize();
 }
 
+TOnlineResult<FGetResolvedConnectString> FOnlineServicesOSSAdapter::GetResolvedConnectString(FGetResolvedConnectString::Params&& Params)
+{
+	FSessionsOSSAdapter* SessionsOSSAdapter = Get<FSessionsOSSAdapter>();
+	check(SessionsOSSAdapter);
+
+	return SessionsOSSAdapter->GetResolvedConnectString(Params);
+}
+
 /* UE::Online*/ }
diff --git a/Engine/Plugins/Online/Experimental/OnlineServicesOSSAdapter/Source/Private/Online/OnlineServicesOSSAdapterModule.cpp b/Engine/Plugins/Online/Experimental/OnlineServicesOSSAdapter/Source/Private/Online/OnlineServicesOSSAdapterModule.cpp
--- a/Engine/Plugins/Online/Experimental/OnlineServicesOSSAdapter/Source/Private/Online/OnlineServicesOSSAdapterModule.cpp
+++ b/Engine/Plugins/Online/Experimental/OnlineServicesOSSAdapter/Source/Private/Online/OnlineServicesOSSAdapterModule.cpp
@@ -7,6 +7,7 @@
 #include "Online/OnlineServicesOSSAdapter.h"
 #include "Online/OnlineServicesRegistry.h"
 #include "Online/OnlineIdOSSAdapter.h"
+#include "Online/SessionsOSSAdapter.h"
 
 #include "OnlineSubsystem.h"
 
@@ -136,7 +137,9 @@ void FOnlineServicesOSSAdapterModule::StartupModule()
 			if (IOnlineSubsystem::IsEnabled(ServiceConfig.OnlineSubsystem))
 			{
 				FOnlineServicesRegistry::Get().RegisterServicesFactory(ServiceConfig.Service, MakeUnique<FOnlineServicesFactoryOSSAdapter>(ServiceConfig), ServiceConfig.Priority);
-				FOnlineIdRegistryRegistry::Get().RegisterAccountIdRegistry(ServiceConfig.Service, new FOnlineUniqueNetIdRegistry(ServiceConfig.Service), ServiceConfig.Priority);
+				FOnlineIdRegistryRegistry::Get().RegisterAccountIdRegistry(ServiceConfig.Service, new FOnlineAccountIdRegistryOSSAdapter(ServiceConfig.Service), ServiceConfig.Priority);
+				FOnlineIdRegistryRegistry::Get().RegisterSessionIdRegistry(ServiceConfig.Service, new FOnlineSessionIdRegistryOSSAdapter(ServiceConfig.Service), ServiceConfig.Priority);
+				FOnlineIdRegistryRegistry::Get().RegisterSessionInviteIdRegistry(ServiceConfig.Service, new FOnlineSessionInviteIdRegistryOSSAdapter(ServiceConfig.Service), ServiceConfig.Priority);
 			}
 		}
 	}
@@ -157,6 +160,8 @@ void FOnlineServicesOSSAdapterModule::ShutdownModule()
 			{
 				FOnlineServicesRegistry::Get().UnregisterServicesFactory(ServiceConfig.Service, ServiceConfig.Priority);
 				FOnlineIdRegistryRegistry::Get().UnregisterAccountIdRegistry(ServiceConfig.Service, ServiceConfig.Priority);
+				FOnlineIdRegistryRegistry::Get().UnregisterSessionIdRegistry(ServiceConfig.Service, ServiceConfig.Priority);
+				FOnlineIdRegistryRegistry::Get().UnregisterSessionInviteIdRegistry(ServiceConfig.Service, ServiceConfig.Priority);
 			}
 		}
 	}
diff --git a/Engine/Plugins/Online/Experimental/OnlineServicesOSSAdapter/Source/Private/Online/SessionsOSSAdapter.cpp b/Engine/Plugins/Online/Experimental/OnlineServicesOSSAdapter/Source/Private/Online/SessionsOSSAdapter.cpp
new file mode 100644
--- /dev/null
+++ b/Engine/Plugins/Online/Experimental/OnlineServicesOSSAdapter/Source/Private/Online/SessionsOSSAdapter.cpp
@@ -0,0 +1,1081 @@
+// Copyright Epic Games, Inc. All Rights Reserved.
+
+#include "Online/SessionsOSSAdapter.h"
+
+#include "OnlineSubsystem.h"
+#include "Interfaces/OnlineSessionInterface.h"
+#include "Interfaces/OnlineIdentityInterface.h"
+
+#include "Online/OnlineServicesOSSAdapter.h"
+#include "Online/AuthOSSAdapter.h"
+#include "Online/UserInfoOSSAdapter.h"
+#include "Online/DelegateAdapter.h"
+
+#include "Algo/Find.h"
+
+namespace UE::Online {
+
+/** Auxiliary functions */
+
+EOnlineComparisonOp::Type GetV1SessionSearchComparisonOp(const ESessionsComparisonOp& InValue)
+{
+	switch (InValue)
+	{
+	case ESessionsComparisonOp::Equals:				return EOnlineComparisonOp::Equals;
+	case ESessionsComparisonOp::GreaterThan:		return EOnlineComparisonOp::GreaterThan;
+	case ESessionsComparisonOp::GreaterThanEquals:	return EOnlineComparisonOp::GreaterThanEquals;
+	case ESessionsComparisonOp::LessThan:			return EOnlineComparisonOp::LessThan;
+	case ESessionsComparisonOp::LessThanEquals:		return EOnlineComparisonOp::LessThanEquals;
+	case ESessionsComparisonOp::Near:				return EOnlineComparisonOp::Near;
+	case ESessionsComparisonOp::NotEquals:			return EOnlineComparisonOp::NotEquals;
+	}
+
+	checkNoEntry();
+
+	return EOnlineComparisonOp::Equals;
+}
+
+ESessionState GetV2SessionState(const EOnlineSessionState::Type InValue)
+{
+	switch (InValue)
+	{
+	case EOnlineSessionState::NoSession: return ESessionState::Invalid;
+	case EOnlineSessionState::Creating: return ESessionState::Creating;
+	case EOnlineSessionState::Pending: return ESessionState::Joining;
+	case EOnlineSessionState::Starting: return ESessionState::Valid;
+	case EOnlineSessionState::InProgress: return ESessionState::Valid;
+	case EOnlineSessionState::Ending: return ESessionState::Valid;
+	case EOnlineSessionState::Ended: return ESessionState::Valid;
+	case EOnlineSessionState::Destroying: return ESessionState::Destroying;
+	}
+
+	checkNoEntry();
+	return ESessionState::Invalid;
+}
+
+FCustomSessionSettingsMap GetV2SessionSettings(const ::FSessionSettings& InSessionSettings)
+{
+	FCustomSessionSettingsMap Result;
+
+	// TODO: Pending SchemaVariant work
+
+	return Result;
+}
+
+TSharedRef<FOnlineSessionSearch> BuildV1SessionSearch(const FFindSessions::Params& SearchParams)
+{
+	TSharedRef<FOnlineSessionSearch> Result = MakeShared<FOnlineSessionSearch>();
+
+	Result->bIsLanQuery = SearchParams.bFindLANSessions;
+
+	Result->MaxSearchResults = SearchParams.MaxResults;
+
+	if (const FFindSessionsSearchFilter* PingBucketSize = Algo::FindBy(SearchParams.Filters, OSS_ADAPTER_SESSION_SEARCH_PING_BUCKET_SIZE, &FFindSessionsSearchFilter::Key))
+	{
+		Result->PingBucketSize = (int32)PingBucketSize->Value.Get<int64>();
+	}
+
+	if (const FFindSessionsSearchFilter* PlatformHash = Algo::FindBy(SearchParams.Filters, OSS_ADAPTER_SESSION_SEARCH_PLATFORM_HASH, &FFindSessionsSearchFilter::Key))
+	{
+		Result->PlatformHash = (int32)PlatformHash->Value.Get<int64>();
+	}
+
+	//for (const FFindSessionsSearchFilter& SearchFilter : SearchParams.Filters)
+	//{
+		// TODO: Pending SchemaVariant support
+		// Result->QuerySettings.Set(SearchFilter.Key, SearchFilter.Value, GetV1SessionSearchComparisonOp(SearchFilter.ComparisonOp));
+	//}
+
+	return Result;
+}
+
+/** FSessionOSSAdapter */
+
+FSessionOSSAdapter::FSessionOSSAdapter(const FOnlineSession& InSession)
+	: V1Session(InSession)
+{
+
+}
+
+const FSessionOSSAdapter& FSessionOSSAdapter::Cast(const FSession& InSession)
+{
+	// TODO: Check for ::Platform type if that's the Adapter type
+
+	return static_cast<const FSessionOSSAdapter&>(InSession);
+}
+
+const FOnlineSession& FSessionOSSAdapter::GetV1Session() const
+{
+	return V1Session;
+}
+
+/** FSessionsOSSAdapter */ 
+
+FSessionsOSSAdapter::FSessionsOSSAdapter(FOnlineServicesOSSAdapter& InServices)
+	: Super(InServices)
+{
+}
+
+void FSessionsOSSAdapter::Initialize()
+{
+	Super::Initialize();
+
+	IOnlineSubsystem& SubsystemV1 = static_cast<FOnlineServicesOSSAdapter&>(Services).GetSubsystem();
+
+	SessionsInterface = SubsystemV1.GetSessionInterface();
+	check(SessionsInterface);
+
+	IdentityInterface = SubsystemV1.GetIdentityInterface();
+	check(IdentityInterface);
+
+	SessionsInterface->AddOnSessionInviteReceivedDelegate_Handle(FOnSessionInviteReceivedDelegate::CreateLambda([this](const FUniqueNetId& UserId, const FUniqueNetId& FromId, const FString& AppId, const FOnlineSessionSearchResult& InviteResult)
+	{
+		FOnlineServicesOSSAdapter& ServicesOSSAdapter = static_cast<FOnlineServicesOSSAdapter&>(Services);
+
+		FSessionInvite SessionInvite;
+		SessionInvite.SenderId = ServicesOSSAdapter.GetAccountIdRegistry().FindOrAddHandle(FromId.AsShared());
+		SessionInvite.RecipientId = ServicesOSSAdapter.GetAccountIdRegistry().FindOrAddHandle(UserId.AsShared());
+		SessionInvite.InviteId = GetSessionInviteIdRegistry().BasicRegistry.FindOrAddHandle(AppId);
+		SessionInvite.Session = BuildV2Session(&InviteResult.Session);
+
+		FSessionInviteReceived Event{ ServicesOSSAdapter.GetAccountIdRegistry().FindOrAddHandle(UserId.AsShared()), MakeShared<FSessionInvite>(MoveTemp(SessionInvite)) };
+
+		SessionEvents.OnSessionInviteReceived.Broadcast(Event);
+	}));
+
+	SessionsInterface->AddOnSessionUserInviteAcceptedDelegate_Handle(FOnSessionUserInviteAcceptedDelegate::CreateLambda([this](const bool bWasSuccessful, const int32 ControllerId, FUniqueNetIdPtr InvitedUserIdPtr, const FOnlineSessionSearchResult& InviteResult)
+	{
+		FOnlineServicesOSSAdapter& ServicesOSSAdapter = static_cast<FOnlineServicesOSSAdapter&>(Services);
+
+		FUniqueNetIdRef LocalUserId = FUniqueNetIdString::EmptyId();
+		FUniqueNetIdPtr LocalUserIdPtr = IdentityInterface->GetUniquePlayerId(ControllerId);
+		if (LocalUserIdPtr.IsValid())
+		{
+			LocalUserId = (*LocalUserIdPtr).AsShared();
+		}
+
+		FUISessionJoinRequested Event {
+			ServicesOSSAdapter.GetAccountIdRegistry().FindOrAddHandle(LocalUserId),
+			TResult<TSharedRef<const FSession>, FOnlineError>(BuildV2Session(&InviteResult.Session)),
+			EUISessionJoinRequestedSource::FromInvitation
+		};
+	}));
+
+	// All delegates handling Settings or Participant changes in the V1 implementation will be handled with calls to UpdateSession, and the FSessionUpdated event
+}
+
+void FSessionsOSSAdapter::Shutdown()
+{
+	SessionsInterface->ClearOnSessionInviteReceivedDelegates(this);
+
+	Super::Shutdown();
+}
+
+TOnlineResult<FGetAllSessions> FSessionsOSSAdapter::GetAllSessions(FGetAllSessions::Params&& Params) const
+{
+	TArray<TSharedRef<const FSession>> Sessions;
+
+	for (const FName SessionName : RegisteredSessionNames)
+	{
+		if (FNamedOnlineSession* NamedSession = SessionsInterface->GetNamedSession(SessionName))
+		{
+			Sessions.Add(BuildV2Session(NamedSession));
+		}
+	}
+
+	return TOnlineResult<FGetAllSessions>({ Sessions });
+}
+
+TOnlineResult<FGetSessionByName> FSessionsOSSAdapter::GetSessionByName(FGetSessionByName::Params&& Params) const
+{
+	if (FNamedOnlineSession* NamedSession = SessionsInterface->GetNamedSession(Params.LocalName))
+	{
+		return TOnlineResult<FGetSessionByName>({ BuildV2Session(NamedSession) });
+	}
+	else
+	{
+		return TOnlineResult<FGetSessionByName>(Errors::NotFound());
+	}
+}
+
+TOnlineResult<FGetSessionById> FSessionsOSSAdapter::GetSessionById(FGetSessionById::Params&& Params) const
+{
+	FOnlineServicesOSSAdapter& ServicesOSSAdapter = static_cast<FOnlineServicesOSSAdapter&>(Services);
+
+	for (const FName SessionName : RegisteredSessionNames)
+	{
+		if (FNamedOnlineSession* NamedSession = SessionsInterface->GetNamedSession(SessionName))
+		{
+			FOnlineSessionIdHandle SessionIdHandle = GetSessionIdRegistry().FindOrAddHandle(NamedSession->SessionInfo->GetSessionId().AsShared());
+
+			if (SessionIdHandle == Params.IdHandle)
+			{
+				return TOnlineResult<FGetSessionById>({ BuildV2Session(NamedSession) });
+			}
+		}
+	}
+
+	return TOnlineResult<FGetSessionById>(Errors::NotFound());
+}
+
+TOnlineAsyncOpHandle<FCreateSession> FSessionsOSSAdapter::CreateSession(FCreateSession::Params&& Params)
+{
+	TOnlineAsyncOpRef<FCreateSession> Op = GetOp<FCreateSession>(MoveTemp(Params));
+	const FCreateSession::Params& OpParams = Op->GetParams();
+
+	FOnlineError ParamsCheck = CheckCreateSessionParams(OpParams);
+	if (ParamsCheck != Errors::Success())
+	{
+		Op->SetError(MoveTemp(ParamsCheck));
+		return Op->GetHandle();
+	}
+
+	Op->Then([this](TOnlineAsyncOp<FCreateSession>& Op)
+	{
+		const FCreateSession::Params& OpParams = Op.GetParams();
+
+		FOnlineError StateCheck = CheckCreateSessionState(OpParams);
+		if (StateCheck != Errors::Success())
+		{
+			Op.SetError(MoveTemp(StateCheck));
+			return;
+		}
+
+		// Dedicated servers can update sessions so we won't check if LocalUserId is valid
+		const FUniqueNetIdRef LocalUserId = Services.Get<FAuthOSSAdapter>()->GetUniqueNetId(OpParams.LocalUserId).ToSharedRef();
+
+		MakeMulticastAdapter(this, SessionsInterface->OnCreateSessionCompleteDelegates,
+		[this, WeakOp = Op.AsWeak()](FName SessionName, const bool bWasSuccessful) mutable
+		{
+			if (TOnlineAsyncOpPtr<FCreateSession> Op = WeakOp.Pin())
+			{
+				if (!bWasSuccessful)
+				{
+					Op->SetError(Errors::Unknown());
+					return;
+				}
+
+				RegisteredSessionNames.Add(SessionName);
+
+				const FNamedOnlineSession* V1Session = SessionsInterface->GetNamedSession(SessionName);
+
+				TSharedRef<FSession> V2Session = BuildV2Session(V1Session);
+
+				Op->SetResult({ V2Session });
+			}
+		});
+
+		SessionsInterface->CreateSession(*LocalUserId, OpParams.SessionName, BuildV1Settings(OpParams.SessionSettings));
+
+		// TODO: Pending support for multiple local users
+	})
+	.Enqueue(GetSerialQueue());
+
+	return Op->GetHandle();
+}
+
+TOnlineAsyncOpHandle<FUpdateSession> FSessionsOSSAdapter::UpdateSession(FUpdateSession::Params&& Params)
+{
+	TOnlineAsyncOpRef<FUpdateSession> Op = GetOp<FUpdateSession>(MoveTemp(Params));
+	const FUpdateSession::Params& OpParams = Op->GetParams();
+
+	Op->Then([this](TOnlineAsyncOp<FUpdateSession>& Op)
+	{
+		const FUpdateSession::Params& OpParams = Op.GetParams();
+
+		FOnlineError StateCheck = CheckUpdateSessionState(OpParams);
+		if (StateCheck != Errors::Success())
+		{
+			Op.SetError(MoveTemp(StateCheck));
+			return;
+		}
+
+		MakeMulticastAdapter(this, SessionsInterface->OnUpdateSessionCompleteDelegates,
+		[this, WeakOp = Op.AsWeak()](FName SessionName, const bool bWasSuccessful)
+		{
+			if (TOnlineAsyncOpPtr<FUpdateSession> Op = WeakOp.Pin())
+			{
+				if (!bWasSuccessful)
+				{
+					Op->SetError(Errors::Unknown());
+					return;
+				}
+
+				const FUpdateSession::Params& OpParams = Op->GetParams();
+
+				if (TSharedRef<FSession>* FoundSession = LocalSessionsByName.Find(OpParams.SessionName))
+				{
+					Op->SetResult({ *FoundSession });
+
+					FSessionUpdated Event { *FoundSession, OpParams.Mutations };
+
+					SessionEvents.OnSessionUpdated.Broadcast(Event);
+				}
+				else
+				{
+					Op->SetError(Errors::NotFound());
+				}
+			}
+		});
+
+		const TSharedRef<FSession>& FoundSession = LocalSessionsByName.FindChecked(OpParams.SessionName);
+
+		// We will update a copy here, and wait until the operation has completed successfully to update our local data
+		FSessionSettings UpdatedV2Settings = FoundSession->SessionSettings;
+		UpdatedV2Settings += OpParams.Mutations;
+		
+		FOnlineSessionSettings UpdatedV1Settings = BuildV1Settings(UpdatedV2Settings);
+
+		SessionsInterface->UpdateSession(OpParams.SessionName, UpdatedV1Settings);
+	})
+	.Enqueue(GetSerialQueue());
+
+	return Op->GetHandle();
+}
+
+TOnlineAsyncOpHandle<FLeaveSession> FSessionsOSSAdapter::LeaveSession(FLeaveSession::Params&& Params)
+{
+	TOnlineAsyncOpRef<FLeaveSession> Op = GetOp<FLeaveSession>(MoveTemp(Params));
+	const FLeaveSession::Params& OpParams = Op->GetParams();
+
+	Op->Then([this](TOnlineAsyncOp<FLeaveSession>& Op)
+	{
+		const FLeaveSession::Params& OpParams = Op.GetParams();
+
+		FOnlineError StateCheck = CheckLeaveSessionState(OpParams);
+		if (StateCheck != Errors::Success())
+		{
+			Op.SetError(MoveTemp(StateCheck));
+			return;
+		}
+
+		MakeMulticastAdapter(this, SessionsInterface->OnDestroySessionCompleteDelegates,
+		[this, WeakOp = Op.AsWeak()](FName SessionName, const bool bWasSuccessful) mutable
+		{
+			if (TOnlineAsyncOpPtr<FLeaveSession> Op = WeakOp.Pin())
+			{
+				if (!bWasSuccessful)
+				{
+					Op->SetError(Errors::Unknown());
+					SessionsInterface->ClearOnDestroySessionCompleteDelegates(this);
+					return;
+				}
+
+				RegisteredSessionNames.Remove(SessionName);
+
+				Op->SetResult({ });
+
+				FSessionLeft SessionLeftEvent;
+				SessionLeftEvent.LocalUserIds.Append(Op->GetParams().LocalUsers);
+
+				SessionEvents.OnSessionLeft.Broadcast(SessionLeftEvent);
+			}
+
+			SessionsInterface->ClearOnDestroySessionCompleteDelegates(this);
+		});
+
+		SessionsInterface->DestroySession(OpParams.SessionName);
+
+		// TODO: Pending support for multiple local users
+	})
+	.Enqueue(GetSerialQueue());
+
+	return Op->GetHandle();
+}
+
+TOnlineAsyncOpHandle<FFindSessions> FSessionsOSSAdapter::FindSessions(FFindSessions::Params&& Params)
+{
+	TOnlineAsyncOpRef<FFindSessions> Op = GetOp<FFindSessions>(MoveTemp(Params));
+	const FFindSessions::Params& OpParams = Op->GetParams();
+
+	FOnlineError ParamsCheck = CheckFindSessionsParams(OpParams);
+	if (ParamsCheck != Errors::Success())
+	{
+		Op->SetError(MoveTemp(ParamsCheck));
+		return Op->GetHandle();
+	}
+
+	Op->Then([this](TOnlineAsyncOp<FFindSessions>& Op)
+	{
+		const FFindSessions::Params& OpParams = Op.GetParams();
+
+		FOnlineError StateCheck = CheckFindSessionsState(OpParams);
+		if (StateCheck != Errors::Success())
+		{
+			Op.SetError(MoveTemp(StateCheck));
+			return;
+		}
+
+		const FUniqueNetIdRef LocalUserId = Services.Get<FAuthOSSAdapter>()->GetUniqueNetId(OpParams.LocalUserId).ToSharedRef();
+		if (!LocalUserId->IsValid())
+		{
+			Op.SetError(Errors::InvalidUser());
+			return;
+		}
+
+		if (PendingV1SessionSearchesPerUser.Contains(OpParams.LocalUserId))
+		{
+			Op.SetError(Errors::AlreadyPending());
+			return;
+		}
+
+		// If there is a target user, we will use FindFriendSession
+		if (OpParams.SessionId.IsSet())
+		{
+			if (const FUniqueNetIdPtr SessionIdPtr = GetSessionIdRegistry().GetIdValue(*OpParams.SessionId))
+			{
+				const FUniqueNetIdRef SessionIdRef = SessionIdPtr.ToSharedRef();
+
+				MakeMulticastAdapter(this, SessionsInterface->OnFindSessionsCompleteDelegates,
+				[this, WeakOp = Op.AsWeak()](bool bWasSuccessful)
+				{
+					if (TOnlineAsyncOpPtr<FFindSessions> Op = WeakOp.Pin())
+					{
+						if (bWasSuccessful)
+						{
+							TSharedRef<FOnlineSessionSearch>& PendingV1SessionSearch = PendingV1SessionSearchesPerUser.FindChecked(Op->GetParams().LocalUserId);
+
+							TArray<TSharedRef<const FSession>> FoundSessions = BuildV2SessionSearchResults(PendingV1SessionSearch->SearchResults);
+
+							Op->SetResult({ MoveTemp(FoundSessions) });
+						}
+						else
+						{
+							Op->SetError(Errors::Unknown());
+						}
+
+						PendingV1SessionSearchesPerUser.Remove(Op->GetParams().LocalUserId);
+					}
+				});
+
+				FOnlineServicesOSSAdapter& ServicesOSSAdapter = static_cast<FOnlineServicesOSSAdapter&>(Services);
+
+				FUniqueNetIdRef TargetIdRef = FUniqueNetIdString::EmptyId();
+				if (OpParams.TargetUser.IsSet())
+				{
+					TargetIdRef = ServicesOSSAdapter.GetAccountIdRegistry().GetIdValue(*OpParams.TargetUser).ToSharedRef();
+				}
+
+				PendingV1SessionSearchesPerUser.Emplace(OpParams.LocalUserId, MakeShared<FOnlineSessionSearch>());
+
+				SessionsInterface->FindSessionById(*LocalUserId, *SessionIdRef, *TargetIdRef, *MakeDelegateAdapter(this, [this, WeakOp = Op.AsWeak()](int32 LocalUserNum, bool bWasSuccessful, const FOnlineSessionSearchResult& SingleSearchResult) mutable
+				{
+					if (TOnlineAsyncOpPtr<FFindSessions> Op = WeakOp.Pin())
+					{
+						if (!bWasSuccessful)
+						{
+							Op->SetError(Errors::Unknown());
+							return;
+						}
+
+						TSharedRef<FOnlineSessionSearch>& PendingV1SessionSearch = PendingV1SessionSearchesPerUser.FindChecked(Op->GetParams().LocalUserId);
+
+						PendingV1SessionSearch->SearchResults.Add(SingleSearchResult);
+					}
+				}));
+			}
+			else
+			{
+				Op.SetError(Errors::InvalidParams());
+				return;
+			}
+		}
+		else if (OpParams.TargetUser.IsSet())
+		{
+			int32 LocalUserNum = Services.Get<FAuthOSSAdapter>()->GetLocalUserNum(OpParams.LocalUserId);
+
+			MakeMulticastAdapter(this, SessionsInterface->OnFindFriendSessionCompleteDelegates[LocalUserNum],
+			[this, WeakOp = Op.AsWeak()](int32 LocalUserId, bool bWasSuccessful, const TArray<FOnlineSessionSearchResult>& FriendSearchResults)
+			{
+				if (TOnlineAsyncOpPtr<FFindSessions> Op = WeakOp.Pin())
+				{
+					if (bWasSuccessful)
+					{
+						TArray<TSharedRef<const FSession>> FoundSessions = BuildV2SessionSearchResults(FriendSearchResults);
+
+						Op->SetResult({ MoveTemp(FoundSessions) });
+					}
+					else
+					{
+						Op->SetError(Errors::Unknown());						
+					}
+
+					PendingV1SessionSearchesPerUser.Remove(Op->GetParams().LocalUserId);
+				}
+			});
+
+			FOnlineServicesOSSAdapter& ServicesOSSAdapter = static_cast<FOnlineServicesOSSAdapter&>(Services);
+
+			TArray<FUniqueNetIdRef> TargetUniqueNetIds = { ServicesOSSAdapter.GetAccountIdRegistry().GetIdValue(*OpParams.TargetUser).ToSharedRef() };
+
+			// The cached search information is not used in this mode, but we'll still update it to not have more than one search at the same time
+			PendingV1SessionSearchesPerUser.Emplace(OpParams.LocalUserId, MakeShared<FOnlineSessionSearch>());
+
+			SessionsInterface->FindFriendSession(*LocalUserId, TargetUniqueNetIds);
+		}
+		else
+		{
+			MakeMulticastAdapter(this, SessionsInterface->OnFindSessionsCompleteDelegates,
+			[this, WeakOp = Op.AsWeak()](bool bWasSuccessful) mutable
+			{
+				if (TOnlineAsyncOpPtr<FFindSessions> Op = WeakOp.Pin())
+				{
+					if (bWasSuccessful)
+					{
+						TSharedRef<FOnlineSessionSearch>& PendingV1SessionSearch = PendingV1SessionSearchesPerUser.FindChecked(Op->GetParams().LocalUserId);
+
+						TArray<TSharedRef<const FSession>> FoundSessions = BuildV2SessionSearchResults(PendingV1SessionSearch->SearchResults);
+
+						Op->SetResult({ MoveTemp(FoundSessions) });
+					}
+					else
+					{
+						Op->SetError(Errors::Unknown());						
+					}
+
+					PendingV1SessionSearchesPerUser.Remove(Op->GetParams().LocalUserId);
+				}
+			});
+
+			TSharedRef<FOnlineSessionSearch>& PendingV1SessionSearch = PendingV1SessionSearchesPerUser.Emplace(OpParams.LocalUserId, BuildV1SessionSearch(OpParams));
+
+			SessionsInterface->FindSessions(*LocalUserId, PendingV1SessionSearch);
+		}
+	})
+	.Enqueue(GetSerialQueue());
+
+	return Op->GetHandle();
+}
+
+TOnlineAsyncOpHandle<FStartMatchmaking> FSessionsOSSAdapter::StartMatchmaking(FStartMatchmaking::Params&& Params)
+{
+	TOnlineAsyncOpRef<FStartMatchmaking> Op = GetOp<FStartMatchmaking>(MoveTemp(Params));
+	const FStartMatchmaking::Params& OpParams = Op->GetParams();
+
+	FOnlineError ParamsCheck = CheckStartMatchmakingParams(OpParams);
+	if (ParamsCheck != Errors::Success())
+	{
+		Op->SetError(MoveTemp(ParamsCheck));
+		return Op->GetHandle();
+	}
+
+	Op->Then([this](TOnlineAsyncOp<FStartMatchmaking>& Op)
+	{
+		const FStartMatchmaking::Params& OpParams = Op.GetParams();
+
+		FOnlineError StateCheck = CheckStartMatchmakingState(OpParams);
+		if (StateCheck != Errors::Success())
+		{
+			Op.SetError(MoveTemp(StateCheck));
+			return;
+		}
+
+		if (PendingV1SessionSearchesPerUser.Contains(OpParams.LocalUserId))
+		{
+			Op.SetError(Errors::AlreadyPending());
+			return;
+		}
+
+		FOnlineServicesOSSAdapter& ServicesOSSAdapter = static_cast<FOnlineServicesOSSAdapter&>(Services);
+
+		TArray<FSessionMatchmakingUser> MatchMakingUsers;
+		for (const TPair<FOnlineAccountIdHandle, FSessionMember>& LocalUser : OpParams.LocalUsers)
+		{
+			FSessionMatchmakingUser NewMatchMakingUser { ServicesOSSAdapter.GetAccountIdRegistry().GetIdValue(LocalUser.Key).ToSharedRef() };
+
+			// TODO: Copy Attributes from session members parameters (Pending SchemaVariant)
+
+			MatchMakingUsers.Add(NewMatchMakingUser);
+		}
+
+		FOnlineSessionSettings NewSessionSettings = BuildV1Settings(OpParams.SessionSettings);
+
+		/** We build a mock params struct for FindSessions and convert it to a V1 session search */
+		FFindSessions::Params FindSessionsParams;
+		FindSessionsParams.bFindLANSessions = false;
+		FindSessionsParams.Filters = OpParams.SearchFilters;
+		FindSessionsParams.LocalUserId = OpParams.LocalUserId;
+		FindSessionsParams.MaxResults = 1;
+
+		TSharedRef<FOnlineSessionSearch>& SessionSearch = PendingV1SessionSearchesPerUser.Emplace(OpParams.LocalUserId, BuildV1SessionSearch(FindSessionsParams));
+
+		// QUESTION: StartMatchmaking causes both OnMatchmakingComplete and OnStartMatchmakingComplete to trigger.
+		// The first triggers the second so I used the latter, but the interface header did not specify if one should be used. Is this the correct one?
+		SessionsInterface->StartMatchmaking(MatchMakingUsers, OpParams.SessionName, NewSessionSettings, SessionSearch, *MakeDelegateAdapter(this, [this, WeakOp = Op.AsWeak()](FName SessionName, const ::FOnlineError& ErrorDetails, const FSessionMatchmakingResults& Results) mutable
+		{
+			if (TOnlineAsyncOpPtr<FStartMatchmaking> Op = WeakOp.Pin())
+			{
+				if (!ErrorDetails.bSucceeded)
+				{
+					Op->SetError(Errors::Unknown()); // TODO: Proper error parsing
+					return;
+				}
+
+				TSharedRef<FSession> NewSession = BuildV2Session(SessionsInterface->GetNamedSession(SessionName));
+
+				Op->SetResult({ NewSession });
+
+				FSessionJoined Event{ { Op->GetParams().LocalUserId }, NewSession };
+
+				SessionEvents.OnSessionJoined.Broadcast(Event);
+
+				// TODO: Pending multiple local user support
+			}
+		}));
+	})
+	.Enqueue(GetSerialQueue());
+
+	return Op->GetHandle();
+}
+
+TOnlineAsyncOpHandle<FJoinSession> FSessionsOSSAdapter::JoinSession(FJoinSession::Params&& Params)
+{
+	TOnlineAsyncOpRef<FJoinSession> Op = GetOp<FJoinSession>(MoveTemp(Params));
+	const FJoinSession::Params& OpParams = Op->GetParams();
+
+	FOnlineError ParamsCheck = CheckJoinSessionParams(OpParams);
+	if (ParamsCheck != Errors::Success())
+	{
+		Op->SetError(MoveTemp(ParamsCheck));
+		return Op->GetHandle();
+	}
+
+	Op->Then([this](TOnlineAsyncOp<FJoinSession>& Op)
+	{
+		const FJoinSession::Params& OpParams = Op.GetParams();
+
+		FOnlineError StateCheck = CheckJoinSessionState(OpParams);
+		if (StateCheck != Errors::Success())
+		{
+			Op.SetError(MoveTemp(StateCheck));
+			return;
+		}
+
+		const FUniqueNetIdRef LocalUserId = Services.Get<FAuthOSSAdapter>()->GetUniqueNetId(OpParams.LocalUserId).ToSharedRef();
+		if (!LocalUserId->IsValid())
+		{
+			Op.SetError(Errors::InvalidUser());
+			return;
+		}
+
+		MakeMulticastAdapter(this, SessionsInterface->OnJoinSessionCompleteDelegates,
+		[this, WeakOp = Op.AsWeak()](FName SessionName, EOnJoinSessionCompleteResult::Type Result) mutable
+		{
+			if (TOnlineAsyncOpPtr<FJoinSession> Op = WeakOp.Pin())
+			{
+				if (Result != EOnJoinSessionCompleteResult::Success)
+				{
+					Op->SetError(Errors::Unknown());
+					return;
+				}
+
+				const FJoinSession::Params& OpParams = Op->GetParams();
+
+				TOnlineResult<FGetSessionById> GetSessionByIdResult = GetSessionById({ OpParams.LocalUserId, OpParams.SessionId });
+				if (GetSessionByIdResult.IsError())
+				{
+					// If no result is found, the id might be expired, which we should notify
+					if (GetSessionIdRegistry().IsHandleExpired(OpParams.SessionId))
+					{
+						UE_LOG(LogTemp, Warning, TEXT("[FSessionsOSSAdapter::JoinSession] SessionId parameter [%s] is expired. Please call FindSessions to get an updated list of available sessions "), *ToLogString(OpParams.SessionId));
+					}
+
+					Op->SetError(MoveTemp(GetSessionByIdResult.GetErrorValue()));
+					return;
+				}
+
+				RegisteredSessionNames.Add(SessionName);
+
+				const TSharedRef<const FSession>& FoundSession = GetSessionByIdResult.GetOkValue().Session;
+
+				Op->SetResult({ FoundSession });
+
+				FSessionJoined Event{ { OpParams.LocalUserId }, FoundSession };
+
+				SessionEvents.OnSessionJoined.Broadcast(Event);
+			}
+		});
+
+		TOnlineResult<FGetSessionById> GetSessionByIdResult = GetSessionById({ OpParams.LocalUserId, OpParams.SessionId });
+		if (GetSessionByIdResult.IsError())
+		{
+			Op.SetError(MoveTemp(GetSessionByIdResult.GetErrorValue()));
+			return;
+		}
+
+		const TSharedRef<const FSession>& FoundSession = GetSessionByIdResult.GetOkValue().Session;
+
+		FOnlineSessionSearchResult SearchResult;
+
+		if (const FCustomSessionSetting* PingInMs = FoundSession->SessionSettings.CustomSettings.Find(OSS_ADAPTER_SESSIONS_PING_IN_MS))
+		{
+			SearchResult.PingInMs = (int32)PingInMs->Data.Get<int64>();
+		}
+
+		SearchResult.Session = BuildV1Session(FoundSession);
+
+		SessionsInterface->JoinSession(*LocalUserId, OpParams.SessionName, SearchResult);
+
+		// TODO: Pending support for multiple local users
+	})
+	.Enqueue(GetSerialQueue());
+
+	return Op->GetHandle();
+}
+
+TOnlineAsyncOpHandle<FAddSessionMembers> FSessionsOSSAdapter::AddSessionMembers(FAddSessionMembers::Params&& Params)
+{
+	TOnlineAsyncOpRef<FAddSessionMembers> Op = GetOp<FAddSessionMembers>(MoveTemp(Params));
+	const FAddSessionMembers::Params& OpParams = Op->GetParams();
+
+	Op->Then([this](TOnlineAsyncOp<FAddSessionMembers>& Op)
+	{
+		const FAddSessionMembers::Params& OpParams = Op.GetParams();
+
+		FOnlineError StateCheck = CheckAddSessionMembersState(OpParams);
+		if (StateCheck != Errors::Success())
+		{
+			Op.SetError(MoveTemp(StateCheck));
+			return;
+		}
+
+		SessionsInterface->AddOnRegisterPlayersCompleteDelegate_Handle(*MakeDelegateAdapter(this, [this, WeakOp = Op.AsWeak()](FName SessionName, const TArray<FUniqueNetIdRef>& Players, bool bWasSuccessful) mutable
+		{
+			// Clearing the delegate needs to be the last thing we do, as modifying captured variables after it would result in a crash
+
+			if (TOnlineAsyncOpPtr<FAddSessionMembers> Op = WeakOp.Pin())
+			{
+				if (bWasSuccessful)
+				{
+					Op->SetResult({ });
+				}
+				else
+				{
+					Op->SetError(Errors::Unknown());
+				}
+			}
+
+			SessionsInterface->ClearOnRegisterPlayersCompleteDelegates(this);
+		}));
+
+		FAuthOSSAdapter* Auth = Services.Get<FAuthOSSAdapter>();
+		TArray<FUniqueNetIdRef> TargetUserNetIds;
+		TArray<FOnlineAccountIdHandle> NewSessionMemberIds;
+		OpParams.NewSessionMembers.GenerateKeyArray(NewSessionMemberIds);
+		for (const FOnlineAccountIdHandle& TargetUser : NewSessionMemberIds)
+		{
+			FUniqueNetIdRef TargetUserNetId = Auth->GetUniqueNetId(TargetUser).ToSharedRef();
+			if (TargetUserNetId->IsValid())
+			{
+				TargetUserNetIds.Add(TargetUserNetId);
+			}
+		}
+
+		SessionsInterface->RegisterPlayers(OpParams.SessionName, TargetUserNetIds);
+	})
+	.Enqueue(GetSerialQueue());
+
+	return Op->GetHandle();
+}
+
+TOnlineAsyncOpHandle<FRemoveSessionMembers> FSessionsOSSAdapter::RemoveSessionMembers(FRemoveSessionMembers::Params&& Params)
+{
+	TOnlineAsyncOpRef<FRemoveSessionMembers> Op = GetOp<FRemoveSessionMembers>(MoveTemp(Params));
+	const FRemoveSessionMembers::Params& OpParams = Op->GetParams();
+
+	Op->Then([this](TOnlineAsyncOp<FRemoveSessionMembers>& Op)
+	{
+		const FRemoveSessionMembers::Params& OpParams = Op.GetParams();
+
+		FOnlineError StateCheck = CheckRemoveSessionMembersState(OpParams);
+		if (StateCheck != Errors::Success())
+		{
+			Op.SetError(MoveTemp(StateCheck));
+			return;
+		}
+
+		SessionsInterface->AddOnUnregisterPlayersCompleteDelegate_Handle(*MakeDelegateAdapter(this, [this, WeakOp = Op.AsWeak()](FName SessionName, const TArray<FUniqueNetIdRef>& Players, bool bWasSuccessful) mutable
+		{
+			// Clearing the delegate needs to be the last thing we do, as modifying captured variables after it would result in a crash
+
+			if (TOnlineAsyncOpPtr<FRemoveSessionMembers> Op = WeakOp.Pin())
+			{
+				if (bWasSuccessful)
+				{
+					Op->SetResult({ });
+				}
+				else
+				{
+					Op->SetError(Errors::Unknown());
+				}
+			}
+
+			SessionsInterface->ClearOnUnregisterPlayersCompleteDelegates(this);
+		}));
+
+		FAuthOSSAdapter* Auth = Services.Get<FAuthOSSAdapter>();
+		TArray<FUniqueNetIdRef> TargetUserNetIds;
+		for (const FOnlineAccountIdHandle& TargetUser : OpParams.SessionMemberIds)
+		{
+			FUniqueNetIdRef TargetUserNetId = Auth->GetUniqueNetId(TargetUser).ToSharedRef();
+			if (TargetUserNetId->IsValid())
+			{
+				TargetUserNetIds.Add(TargetUserNetId);
+			}
+		}
+
+		SessionsInterface->UnregisterPlayers(OpParams.SessionName, TargetUserNetIds);
+	})
+	.Enqueue(GetSerialQueue());
+
+	return Op->GetHandle();
+}
+
+TOnlineAsyncOpHandle<FSendSessionInvite> FSessionsOSSAdapter::SendSessionInvite(FSendSessionInvite::Params&& Params)
+{
+	TOnlineAsyncOpRef<FSendSessionInvite> Op = GetOp<FSendSessionInvite>(MoveTemp(Params));
+	const FSendSessionInvite::Params& OpParams = Op->GetParams();
+
+	Op->Then([this](TOnlineAsyncOp<FSendSessionInvite>& Op)
+	{
+		const FSendSessionInvite::Params& OpParams = Op.GetParams();
+
+		FOnlineError StateCheck = CheckSendSessionInviteState(OpParams);
+		if (StateCheck != Errors::Success())
+		{
+			Op.SetError(MoveTemp(StateCheck));
+			return;
+		}
+
+		FAuthOSSAdapter* Auth = Services.Get<FAuthOSSAdapter>();
+
+		const FUniqueNetIdRef LocalUserId = Auth->GetUniqueNetId(OpParams.LocalUserId).ToSharedRef();
+		if (!LocalUserId->IsValid())
+		{
+			Op.SetError(Errors::InvalidUser());
+			return;
+		}
+
+		TArray<FUniqueNetIdRef> TargetUserNetIds;
+		for (const FOnlineAccountIdHandle& TargetUser : OpParams.TargetUsers)
+		{
+			FUniqueNetIdRef TargetUserNetId = Auth->GetUniqueNetId(TargetUser).ToSharedRef();
+			if (TargetUserNetId->IsValid())
+			{
+				TargetUserNetIds.Add(TargetUserNetId);
+			}
+		}
+
+		SessionsInterface->SendSessionInviteToFriends(*LocalUserId, OpParams.SessionName, TargetUserNetIds);
+
+		Op.SetResult({ });
+	})
+	.Enqueue(GetSerialQueue());
+
+	return Op->GetHandle();
+}
+
+TOnlineAsyncOpHandle<FRegisterPlayers> FSessionsOSSAdapter::RegisterPlayers(FRegisterPlayers::Params&& Params)
+{
+	TOnlineAsyncOpRef<FRegisterPlayers> Op = GetOp<FRegisterPlayers>(MoveTemp(Params));
+	Op->SetError(Errors::NotImplemented());
+	return Op->GetHandle();
+}
+
+TOnlineAsyncOpHandle<FUnregisterPlayers> FSessionsOSSAdapter::UnregisterPlayers(FUnregisterPlayers::Params&& Params)
+{
+	TOnlineAsyncOpRef<FUnregisterPlayers> Op = GetOp<FUnregisterPlayers>(MoveTemp(Params));
+	Op->SetError(Errors::NotImplemented());
+	return Op->GetHandle();
+}
+
+TOnlineResult<FGetResolvedConnectString> FSessionsOSSAdapter::GetResolvedConnectString(const FGetResolvedConnectString::Params& Params)
+{
+	if (Params.SessionId.IsValid())
+	{
+		TOnlineResult<FGetSessionById> Result = GetSessionById({ Params.LocalUserId, Params.SessionId });
+		if (Result.IsOk())
+		{
+			FOnlineSessionSearchResult SearchResult;
+			SearchResult.Session = BuildV1Session(Result.GetOkValue().Session);
+
+			FString ConnectString;
+
+			if (SessionsInterface->GetResolvedConnectString(SearchResult, Params.PortType, ConnectString))
+			{
+				return TOnlineResult<FGetResolvedConnectString>({ ConnectString });
+			}
+			else
+			{
+				return TOnlineResult<FGetResolvedConnectString>(Errors::Unknown());
+			}
+		}
+		else
+		{
+			return TOnlineResult<FGetResolvedConnectString>(Result.GetErrorValue());
+		}
+	}
+	else
+	{
+		// No valid session id set
+		return TOnlineResult<FGetResolvedConnectString>(Errors::InvalidParams());
+	}
+}
+
+FOnlineSessionSettings FSessionsOSSAdapter::BuildV1Settings(const FSessionSettings& InSessionSettings) const
+{
+	FOnlineSessionSettings Result;
+
+	Result.bAllowInvites = InSessionSettings.bAllowNewMembers;
+	Result.bAllowJoinInProgress = InSessionSettings.bAllowNewMembers;
+	Result.bAllowJoinViaPresence = InSessionSettings.bAllowNewMembers;
+	Result.bAllowJoinViaPresenceFriendsOnly = InSessionSettings.bAllowNewMembers;
+	Result.Settings.Add(OSS_ADAPTER_SESSIONS_ALLOW_SANCTIONED_PLAYERS, InSessionSettings.bAllowSanctionedPlayers);
+	Result.Settings.Add(OSS_ADAPTER_SESSIONS_ALLOW_UNREGISTERED_PLAYERS, InSessionSettings.bAllowUnregisteredPlayers);
+	Result.bAntiCheatProtected = InSessionSettings.bAntiCheatProtected;
+	Result.bIsDedicated = InSessionSettings.bIsDedicatedServerSession;
+	Result.bIsLANMatch = InSessionSettings.bIsLANSession;
+	Result.bShouldAdvertise = InSessionSettings.bAllowNewMembers;
+	if (const FCustomSessionSetting* BuildUniqueId = InSessionSettings.CustomSettings.Find(OSS_ADAPTER_SESSIONS_BUILD_UNIQUE_ID))
+	{
+		Result.BuildUniqueId = BuildUniqueId->Data.Get<int64>();
+	}
+	if (const FCustomSessionSetting* BuildUniqueId = InSessionSettings.CustomSettings.Find(OSS_ADAPTER_SESSIONS_USE_LOBBIES_IF_AVAILABLE))
+	{
+		Result.bUseLobbiesIfAvailable = BuildUniqueId->Data.Get<bool>();
+	}
+	if (const FCustomSessionSetting* BuildUniqueId = InSessionSettings.CustomSettings.Find(OSS_ADAPTER_SESSIONS_USE_LOBBIES_VOICE_CHAT_IF_AVAILABLE))
+	{
+		Result.bUseLobbiesVoiceChatIfAvailable = BuildUniqueId->Data.Get<bool>();
+	}
+	Result.bUsesPresence = InSessionSettings.bPresenceEnabled;
+	if (const FCustomSessionSetting* BuildUniqueId = InSessionSettings.CustomSettings.Find(OSS_ADAPTER_SESSIONS_USES_STATS))
+	{
+		Result.bUsesStats = BuildUniqueId->Data.Get<bool>();
+	}
+	Result.NumPrivateConnections = InSessionSettings.NumMaxPrivateConnections;
+	Result.NumPublicConnections = InSessionSettings.NumMaxPublicConnections;
+	Result.Settings.Add(OSS_ADAPTER_SESSIONS_SCHEMA_NAME, InSessionSettings.SchemaName.ToString());
+	Result.SessionIdOverride = InSessionSettings.SessionIdOverride;
+
+	// TODO: Pending SchemaVariant work
+	//Result.Settings =
+	
+	FOnlineServicesOSSAdapter& ServicesOSSAdapter = static_cast<FOnlineServicesOSSAdapter&>(Services);
+
+	for (const TPair<FOnlineAccountIdHandle, FSessionMember>& Entry : InSessionSettings.SessionMembers)
+	{
+		Result.MemberSettings.Emplace(ServicesOSSAdapter.GetAccountIdRegistry().GetIdValue(Entry.Key).ToSharedRef(), ::FSessionSettings()); // TODO: Pending SchemaVariant work
+	}
+
+	return Result;
+}
+
+void FSessionsOSSAdapter::WriteV2SessionSettingsFromV1Session(const FOnlineSession* InSession, FSessionSettings& OutSettings) const
+{
+	InSession->SessionSettings.Settings.FindChecked(OSS_ADAPTER_SESSIONS_ALLOW_SANCTIONED_PLAYERS).Data.GetValue(OutSettings.bAllowSanctionedPlayers);
+	InSession->SessionSettings.Settings.FindChecked(OSS_ADAPTER_SESSIONS_ALLOW_UNREGISTERED_PLAYERS).Data.GetValue(OutSettings.bAllowUnregisteredPlayers);
+	OutSettings.bAntiCheatProtected = InSession->SessionSettings.bAntiCheatProtected;
+	OutSettings.bPresenceEnabled = InSession->SessionSettings.bUsesPresence;
+	OutSettings.bIsDedicatedServerSession = InSession->SessionSettings.bIsDedicated;
+	OutSettings.bIsLANSession = InSession->SessionSettings.bIsLANMatch;
+	OutSettings.NumMaxPrivateConnections = InSession->SessionSettings.NumPrivateConnections;
+	OutSettings.NumMaxPublicConnections = InSession->SessionSettings.NumPublicConnections;
+	OutSettings.NumOpenPrivateConnections = InSession->NumOpenPrivateConnections;
+	OutSettings.NumOpenPublicConnections = InSession->NumOpenPublicConnections;
+	FString SchemaNameStr;
+	InSession->SessionSettings.Settings.FindChecked(OSS_ADAPTER_SESSIONS_SCHEMA_NAME).Data.GetValue(SchemaNameStr);
+	OutSettings.SchemaName = FName(SchemaNameStr);
+	OutSettings.SessionIdOverride = InSession->SessionSettings.SessionIdOverride;
+
+	// TODO: Pending SchemaVariant work
+	//Result.CustomSettings = 
+
+	FCustomSessionSetting BuildUniqueId;
+	BuildUniqueId.Data.Set<int64>(InSession->SessionSettings.BuildUniqueId);
+	OutSettings.CustomSettings.Add(OSS_ADAPTER_SESSIONS_BUILD_UNIQUE_ID, BuildUniqueId);
+
+	FOnlineServicesOSSAdapter& ServicesOSSAdapter = static_cast<FOnlineServicesOSSAdapter&>(Services);
+
+	for (TPair<FUniqueNetIdRef, ::FSessionSettings> Entry : InSession->SessionSettings.MemberSettings)
+	{
+		FOnlineAccountIdHandle SessionMemberId = ServicesOSSAdapter.GetAccountIdRegistry().FindOrAddHandle(Entry.Key);
+		FSessionMember SessionMember = FSessionMember(); // TODO: Pending SchemaVariant work
+
+		OutSettings.SessionMembers.Emplace(SessionMemberId, SessionMember);
+	}
+}
+
+void FSessionsOSSAdapter::WriteV2SessionSettingsFromV1NamedSession(const FNamedOnlineSession* InSession, FSessionSettings& OutSettings) const
+{
+	bool bPublicJoinable, bFriendJoinable, bInviteOnly, bAllowInvites;
+	InSession->GetJoinability(bPublicJoinable, bFriendJoinable, bInviteOnly, bAllowInvites);
+
+	OutSettings.bAllowNewMembers = bPublicJoinable || bFriendJoinable || bInviteOnly;
+	OutSettings.JoinPolicy = bInviteOnly ? ESessionJoinPolicy::InviteOnly : bFriendJoinable ? ESessionJoinPolicy::FriendsOnly : ESessionJoinPolicy::Public; // We use Public as the default as closed sessions will have bAllowNewMembers as false
+
+	FOnlineServicesOSSAdapter& ServicesOSSAdapter = static_cast<FOnlineServicesOSSAdapter&>(Services);
+
+	for (const FUniqueNetIdRef& RegisteredPlayer : InSession->RegisteredPlayers)
+	{
+		OutSettings.RegisteredPlayers.Add(ServicesOSSAdapter.GetAccountIdRegistry().FindOrAddHandle(RegisteredPlayer));
+	}
+}
+
+FOnlineSession FSessionsOSSAdapter::BuildV1Session(const TSharedRef<const FSession> InSession) const
+{
+	return FOnlineSession(FSessionOSSAdapter::Cast(*InSession).GetV1Session());
+}
+
+TSharedRef<FSession> FSessionsOSSAdapter::BuildV2Session(const FOnlineSession* InSession) const
+{
+	TSharedRef<FSessionOSSAdapter> Session = MakeShared<FSessionOSSAdapter>(*InSession);
+
+	FOnlineServicesOSSAdapter& ServicesOSSAdapter = static_cast<FOnlineServicesOSSAdapter&>(Services);
+
+	Session->SessionId = GetSessionIdRegistry().FindOrAddHandle(InSession->SessionInfo->GetSessionId().AsShared());
+	WriteV2SessionSettingsFromV1Session(InSession, Session->SessionSettings);
+
+	if (const FNamedOnlineSession* NamedInSession = static_cast<const FNamedOnlineSession*>(InSession))
+	{
+		Session->OwnerUserId = ServicesOSSAdapter.GetAccountIdRegistry().FindOrAddHandle(NamedInSession->LocalOwnerId->AsShared());
+		Session->CurrentState = GetV2SessionState(NamedInSession->SessionState);
+		WriteV2SessionSettingsFromV1NamedSession(NamedInSession, Session->SessionSettings);
+	}
+
+	return Session;
+}
+
+TArray<TSharedRef<const FSession>> FSessionsOSSAdapter::BuildV2SessionSearchResults(const TArray<FOnlineSessionSearchResult>& SessionSearchResults) const
+{
+	TArray<TSharedRef<const FSession>> FoundSessions;
+
+	for (const FOnlineSessionSearchResult& SearchResult : SessionSearchResults)
+	{
+		TSharedRef<FSession> FoundSession = BuildV2Session(&SearchResult.Session);
+
+		FCustomSessionSetting PingInMs;
+		PingInMs.Data.Set<int64>((int64)SearchResult.PingInMs);
+		FoundSession->SessionSettings.CustomSettings.Emplace(OSS_ADAPTER_SESSIONS_PING_IN_MS, PingInMs);
+
+		FoundSessions.Add(FoundSession);
+	}
+
+	return FoundSessions;
+}
+
+FOnlineSessionIdRegistryOSSAdapter& FSessionsOSSAdapter::GetSessionIdRegistry() const
+{
+	FOnlineSessionIdRegistryOSSAdapter* SessionIdRegistry = static_cast<FOnlineSessionIdRegistryOSSAdapter*>(FOnlineIdRegistryRegistry::Get().GetSessionIdRegistry(Services.GetServicesProvider()));
+	check(SessionIdRegistry);
+	return *SessionIdRegistry;
+}
+
+FOnlineSessionInviteIdRegistryOSSAdapter& FSessionsOSSAdapter::GetSessionInviteIdRegistry() const
+{
+	FOnlineSessionInviteIdRegistryOSSAdapter* SessionInviteIdRegistry = static_cast<FOnlineSessionInviteIdRegistryOSSAdapter*>(FOnlineIdRegistryRegistry::Get().GetSessionInviteIdRegistry(Services.GetServicesProvider()));
+	check(SessionInviteIdRegistry);
+	return *SessionInviteIdRegistry;
+}
+
+/* UE::Online */ }
\ No newline at end of file
diff --git a/Engine/Plugins/Online/Experimental/OnlineServicesOSSAdapter/Source/Public/Online/OnlineIdOSSAdapter.h b/Engine/Plugins/Online/Experimental/OnlineServicesOSSAdapter/Source/Public/Online/OnlineIdOSSAdapter.h
--- a/Engine/Plugins/Online/Experimental/OnlineServicesOSSAdapter/Source/Public/Online/OnlineIdOSSAdapter.h
+++ b/Engine/Plugins/Online/Experimental/OnlineServicesOSSAdapter/Source/Public/Online/OnlineIdOSSAdapter.h
@@ -11,19 +11,20 @@ namespace UE::Online {
 /**
  * A net id registry suitable for use with OSS FUniqueNetIds
  */
-class FOnlineUniqueNetIdRegistry : public IOnlineAccountIdRegistry
+template<typename IdType>
+class TOnlineUniqueNetIdRegistry : public IOnlineIdRegistry<IdType>
 {
 public:
-	FOnlineUniqueNetIdRegistry(EOnlineServices InOnlineServicesType)
+	TOnlineUniqueNetIdRegistry(EOnlineServices InOnlineServicesType)
 		: OnlineServicesType(InOnlineServicesType)
 	{
 	}
 
-	virtual ~FOnlineUniqueNetIdRegistry() {}
+	virtual ~TOnlineUniqueNetIdRegistry() {}
 
-	FOnlineAccountIdHandle FindOrAddHandle(const FUniqueNetIdRef& IdValue)
+	TOnlineIdHandle<IdType> FindOrAddHandle(const FUniqueNetIdRef& IdValue)
 	{
-		FOnlineAccountIdHandle Handle;
+		TOnlineIdHandle<IdType> Handle;
 		if (!ensure(IdValue->IsValid()))
 		{
 			return Handle;
@@ -34,7 +35,7 @@ public:
 			FReadScopeLock ReadLock(Lock);
 			if (const uint32* FoundHandle = IdValueToHandleMap.Find(IdValue))
 			{
-				Handle = FOnlineAccountIdHandle(OnlineServicesType, *FoundHandle);
+				Handle = TOnlineIdHandle<IdType>(OnlineServicesType, *FoundHandle);
 			}
 		}
 
@@ -44,13 +45,13 @@ public:
 			FWriteScopeLock WriteLock(Lock);
 			if (const uint32* FoundHandle = IdValueToHandleMap.Find(IdValue))
 			{
-				Handle = FOnlineAccountIdHandle(OnlineServicesType, *FoundHandle);
+				Handle = TOnlineIdHandle<IdType>(OnlineServicesType, *FoundHandle);
 			}
 
 			if (!Handle.IsValid())
 			{
 				IdValues.Emplace(IdValue);
-				Handle = FOnlineAccountIdHandle(OnlineServicesType, IdValues.Num());
+				Handle = TOnlineIdHandle<IdType>(OnlineServicesType, IdValues.Num());
 				IdValueToHandleMap.Emplace(IdValue, Handle.GetHandle());
 			}
 		}
@@ -59,7 +60,7 @@ public:
 	}
 
 	// Returns a copy as it's not thread safe to return a pointer/ref to an element of an array that can be relocated by another thread.
-	FUniqueNetIdPtr GetIdValue(const FOnlineAccountIdHandle Handle) const
+	FUniqueNetIdPtr GetIdValue(const TOnlineIdHandle<IdType> Handle) const
 	{
 		if (Handle.GetOnlineServicesType() == OnlineServicesType && Handle.IsValid())
 		{
@@ -72,26 +73,31 @@ public:
 		return FUniqueNetIdPtr();
 	}
 
-	FUniqueNetIdRef GetIdValueChecked(const FOnlineAccountIdHandle Handle) const
+	FUniqueNetIdRef GetIdValueChecked(const TOnlineIdHandle<IdType> Handle) const
 	{
 		return GetIdValue(Handle).ToSharedRef();
 	}
 
+	bool IsHandleExpired(const FOnlineSessionIdHandle& InHandle) const
+	{
+		return GetIdValue(InHandle).IsValid();
+	}
+
 	// Begin IOnlineAccountIdRegistry
-	virtual FString ToLogString(const FOnlineAccountIdHandle& Handle) const override
+	virtual FString ToLogString(const TOnlineIdHandle<IdType>& Handle) const override
 	{
 		FUniqueNetIdPtr IdValue = GetIdValue(Handle);
 		return IdValue ? IdValue->ToDebugString() : FString(TEXT("invalid_id"));
 	}
 
-	virtual TArray<uint8> ToReplicationData(const FOnlineAccountIdHandle& Handle) const
+	virtual TArray<uint8> ToReplicationData(const TOnlineIdHandle<IdType>& Handle) const override
 	{
 		return TArray<uint8>();
 	}
 
-	virtual FOnlineAccountIdHandle FromReplicationData(const TArray<uint8>& Handle)
+	virtual TOnlineIdHandle<IdType> FromReplicationData(const TArray<uint8>& Handle) override
 	{
-		return FOnlineAccountIdHandle();
+		return TOnlineIdHandle<IdType>();
 	}
 
 	// End IOnlineAccountIdRegistry
@@ -104,4 +110,6 @@ private:
 	TUniqueNetIdMap<uint32> IdValueToHandleMap;
 };
 
+using FOnlineAccountIdRegistryOSSAdapter = TOnlineUniqueNetIdRegistry<OnlineIdHandleTags::FAccount>;
+
 /* UE::Online */ }
diff --git a/Engine/Plugins/Online/Experimental/OnlineServicesOSSAdapter/Source/Public/Online/OnlineServicesOSSAdapter.h b/Engine/Plugins/Online/Experimental/OnlineServicesOSSAdapter/Source/Public/Online/OnlineServicesOSSAdapter.h
--- a/Engine/Plugins/Online/Experimental/OnlineServicesOSSAdapter/Source/Public/Online/OnlineServicesOSSAdapter.h
+++ b/Engine/Plugins/Online/Experimental/OnlineServicesOSSAdapter/Source/Public/Online/OnlineServicesOSSAdapter.h
@@ -4,12 +4,12 @@
 
 #include "Online/OnlineServicesCommon.h"
 
+#include "OnlineIdOSSAdapter.h"
+
 class IOnlineSubsystem;
 
 namespace UE::Online {
 
-class FOnlineUniqueNetIdRegistry;
-
 class FOnlineServicesOSSAdapter : public FOnlineServicesCommon
 {
 public:
@@ -20,15 +20,16 @@ public:
 	ONLINESERVICESOSSADAPTER_API virtual void RegisterComponents() override;
 	ONLINESERVICESOSSADAPTER_API virtual void Initialize() override;
 	virtual EOnlineServices GetServicesProvider() const override { return ServicesType; }
+	virtual TOnlineResult<FGetResolvedConnectString> GetResolvedConnectString(FGetResolvedConnectString::Params&& Params) override;
 
 	IOnlineSubsystem& GetSubsystem() const { return *Subsystem; }
-	FOnlineUniqueNetIdRegistry& GetAccountIdRegistry() const { return *AccountIdRegistry; }
-	FOnlineUniqueNetIdRegistry& GetAccountIdRegistry() { return *AccountIdRegistry; }
+	FOnlineAccountIdRegistryOSSAdapter& GetAccountIdRegistry() const { return *AccountIdRegistry; }
+	FOnlineAccountIdRegistryOSSAdapter& GetAccountIdRegistry() { return *AccountIdRegistry; }
 
 protected:
 	EOnlineServices ServicesType;
 	IOnlineSubsystem* Subsystem;
-	FOnlineUniqueNetIdRegistry* AccountIdRegistry;
+	FOnlineAccountIdRegistryOSSAdapter* AccountIdRegistry;
 };
 
 /* UE::Online */ }
diff --git a/Engine/Plugins/Online/Experimental/OnlineServicesOSSAdapter/Source/Public/Online/SessionsOSSAdapter.h b/Engine/Plugins/Online/Experimental/OnlineServicesOSSAdapter/Source/Public/Online/SessionsOSSAdapter.h
new file mode 100644
--- /dev/null
+++ b/Engine/Plugins/Online/Experimental/OnlineServicesOSSAdapter/Source/Public/Online/SessionsOSSAdapter.h
@@ -0,0 +1,111 @@
+// Copyright Epic Games, Inc. All Rights Reserved.
+
+#pragma once
+
+#include "Online/SessionsCommon.h"
+
+#include "OnlineSubsystemTypes.h"
+#include "OnlineSessionSettings.h"
+
+#include "Online/OnlineIdOSSAdapter.h"
+
+using IOnlineSessionPtr = TSharedPtr<class IOnlineSession>;
+using IOnlineIdentityPtr = TSharedPtr<class IOnlineIdentity>;
+
+namespace UE::Online {
+
+class FOnlineServicesOSSAdapter;
+
+/** Registry of V1 session ids in FUniqueNetIdRef format, indexed by FOnlineSessionIdHandle */
+using FOnlineSessionIdRegistryOSSAdapter = TOnlineUniqueNetIdRegistry<OnlineIdHandleTags::FSession>;
+
+using FOnlineSessionInviteIdRegistryOSSAdapter = FOnlineSessionInviteIdStringRegistry;
+
+static FName OSS_ADAPTER_SESSIONS_ALLOW_SANCTIONED_PLAYERS = TEXT("OSS_ADAPTER_SESSIONS_ALLOW_SANCTIONED_PLAYERS");
+static FName OSS_ADAPTER_SESSIONS_ALLOW_UNREGISTERED_PLAYERS = TEXT("OSS_ADAPTER_SESSIONS_ALLOW_UNREGISTERED_PLAYERS");
+static FName OSS_ADAPTER_SESSIONS_USE_LOBBIES_IF_AVAILABLE = TEXT("OSS_ADAPTER_SESSIONS_USE_LOBBIES_IF_AVAILABLE");
+static FName OSS_ADAPTER_SESSIONS_USE_LOBBIES_VOICE_CHAT_IF_AVAILABLE = TEXT("OSS_ADAPTER_SESSIONS_USE_LOBBIES_VOICE_CHAT_IF_AVAILABLE");
+static FName OSS_ADAPTER_SESSIONS_USES_STATS = TEXT("OSS_ADAPTER_SESSIONS_USES_STATS");
+static FName OSS_ADAPTER_SESSIONS_SCHEMA_NAME = TEXT("OSS_ADAPTER_SESSIONS_SCHEMA_NAME");
+static FName OSS_ADAPTER_SESSIONS_BUILD_UNIQUE_ID = TEXT("OSS_ADAPTER_SESSIONS_BUILD_UNIQUE_ID");
+static FName OSS_ADAPTER_SESSIONS_PING_IN_MS = TEXT("OSS_ADAPTER_SESSIONS_PING_IN_MS");
+
+static FName OSS_ADAPTER_SESSION_SEARCH_PING_BUCKET_SIZE = TEXT("OSS_ADAPTER_SESSION_SEARCH_PING_BUCKET_SIZE");
+static FName OSS_ADAPTER_SESSION_SEARCH_PLATFORM_HASH = TEXT("OSS_ADAPTER_SESSION_SEARCH_PLATFORM_HASH");
+
+class FSessionOSSAdapter : public FSession
+{
+public:
+	FSessionOSSAdapter() = default;
+	FSessionOSSAdapter(const FSessionOSSAdapter& InSession) = default;
+	FSessionOSSAdapter(const FOnlineSession& InSession);
+
+	static const FSessionOSSAdapter& Cast(const FSession& InSession);
+
+	const FOnlineSession& GetV1Session() const;
+
+private:
+	FOnlineSession V1Session;
+};
+
+class FSessionsOSSAdapter : public FSessionsCommon
+{
+public:
+	using Super = FSessionsCommon;
+
+	using FSessionsCommon::FSessionsCommon;
+
+	FSessionsOSSAdapter(FOnlineServicesOSSAdapter& InOwningSubsystem);
+	virtual ~FSessionsOSSAdapter() = default;
+
+	// IOnlineComponent
+	virtual void Initialize() override;
+	virtual void Shutdown() override;
+
+	// ISessions
+	virtual TOnlineResult<FGetAllSessions> GetAllSessions(FGetAllSessions::Params&& Params) const override;
+	virtual TOnlineResult<FGetSessionByName> GetSessionByName(FGetSessionByName::Params&& Params) const override;
+	virtual TOnlineResult<FGetSessionById> GetSessionById(FGetSessionById::Params&& Params) const override;
+	virtual TOnlineAsyncOpHandle<FCreateSession> CreateSession(FCreateSession::Params&& Params) override;
+	virtual TOnlineAsyncOpHandle<FUpdateSession> UpdateSession(FUpdateSession::Params&& Params) override;
+	virtual TOnlineAsyncOpHandle<FLeaveSession> LeaveSession(FLeaveSession::Params&& Params) override;
+	virtual TOnlineAsyncOpHandle<FFindSessions> FindSessions(FFindSessions::Params&& Params) override;
+	virtual TOnlineAsyncOpHandle<FStartMatchmaking> StartMatchmaking(FStartMatchmaking::Params&& Params) override;
+	virtual TOnlineAsyncOpHandle<FJoinSession> JoinSession(FJoinSession::Params&& Params) override;
+	virtual TOnlineAsyncOpHandle<FAddSessionMembers> AddSessionMembers(FAddSessionMembers::Params&& Params) override;
+	virtual TOnlineAsyncOpHandle<FRemoveSessionMembers> RemoveSessionMembers(FRemoveSessionMembers::Params&& Params) override;
+	virtual TOnlineAsyncOpHandle<FSendSessionInvite> SendSessionInvite(FSendSessionInvite::Params&& Params) override;
+	virtual TOnlineAsyncOpHandle<FRegisterPlayers> RegisterPlayers(FRegisterPlayers::Params&& Params) override;
+	virtual TOnlineAsyncOpHandle<FUnregisterPlayers> UnregisterPlayers(FUnregisterPlayers::Params&& Params) override;
+
+	TOnlineResult<FGetResolvedConnectString> GetResolvedConnectString(const FGetResolvedConnectString::Params& Params);
+
+private:
+	/** Builds a V1 Session Settings data type from a V2 equivalent */
+	FOnlineSessionSettings BuildV1Settings(const FSessionSettings& InSessionSettings) const;
+	/** Writes V1 Session Settings information into the passed V2 equivalent */
+	void WriteV2SessionSettingsFromV1Session(const FOnlineSession* InSession, FSessionSettings& OutSettings) const;
+	/** Writes V1 Session Settings information contained in Named session type into the passed V2 equivalent */
+	void WriteV2SessionSettingsFromV1NamedSession(const FNamedOnlineSession* InSession, FSessionSettings& OutSettings) const;
+	/** Builds a V1 Session data type from a V2 equivalent */
+	FOnlineSession BuildV1Session(const TSharedRef<const FSession> InSession) const;
+	/** Builds a V2 Session data type from a V1 equivalent */
+	TSharedRef<FSession> BuildV2Session(const FOnlineSession* InSession) const;
+	/** Builds a V2 Session Search Results array from the passed V1 equivalent types */
+	TArray<TSharedRef<const FSession>> BuildV2SessionSearchResults(const TArray<FOnlineSessionSearchResult>& SessionSearchResults) const;
+
+	FOnlineSessionIdRegistryOSSAdapter& GetSessionIdRegistry() const;
+	FOnlineSessionInviteIdRegistryOSSAdapter& GetSessionInviteIdRegistry() const;
+
+protected:
+	IOnlineSessionPtr SessionsInterface;
+	IOnlineIdentityPtr IdentityInterface;
+
+	/** Record of all names for registered sessions. Needed for session retrieval methods */
+	TSet<FName> RegisteredSessionNames;
+
+	/** Cache of session search information to be used by ongoing session searches, indexed per user. Entries are removed upon search completion. */
+	TMap<FOnlineAccountIdHandle, TSharedRef<FOnlineSessionSearch>> PendingV1SessionSearchesPerUser;
+};
+
+/* UE::Online */ }
diff --git a/Engine/Source/Runtime/CoreOnline/Private/Online/CoreOnline.cpp b/Engine/Source/Runtime/CoreOnline/Private/Online/CoreOnline.cpp
--- a/Engine/Source/Runtime/CoreOnline/Private/Online/CoreOnline.cpp
+++ b/Engine/Source/Runtime/CoreOnline/Private/Online/CoreOnline.cpp
@@ -289,6 +289,122 @@ IOnlineAccountIdRegistry* FOnlineIdRegistryRegistry::GetAccountIdRegistry(EOnlin
 	return Registry;
 }
 
+void FOnlineIdRegistryRegistry::RegisterSessionIdRegistry(EOnlineServices OnlineServices, IOnlineSessionIdRegistry* Registry, int32 Priority)
+{
+	FSessionIdRegistryAndPriority* Found = SessionIdRegistries.Find(OnlineServices);
+	if (Found == nullptr || Found->Priority < Priority)
+	{
+		SessionIdRegistries.Emplace(OnlineServices, FSessionIdRegistryAndPriority(Registry, Priority));
+	}
+}
+
+void FOnlineIdRegistryRegistry::UnregisterSessionIdRegistry(EOnlineServices OnlineServices, int32 Priority)
+{
+	FSessionIdRegistryAndPriority* Found = SessionIdRegistries.Find(OnlineServices);
+	if (Found != nullptr && Found->Priority == Priority)
+	{
+		SessionIdRegistries.Remove(OnlineServices);
+	}
+}
+
+FString FOnlineIdRegistryRegistry::ToLogString(const FOnlineSessionIdHandle& Handle) const
+{
+	FString Result;
+	if (IOnlineSessionIdRegistry* Registry = GetSessionIdRegistry(Handle.GetOnlineServicesType()))
+	{
+		Result = Registry->ToLogString(Handle);
+	}
+	return Result;
+}
+
+TArray<uint8> FOnlineIdRegistryRegistry::ToReplicationData(const FOnlineSessionIdHandle& Handle) const
+{
+	TArray<uint8> RepData;
+	if (IOnlineSessionIdRegistry* Registry = GetSessionIdRegistry(Handle.GetOnlineServicesType()))
+	{
+		RepData = Registry->ToReplicationData(Handle);
+	}
+	return RepData;
+}
+
+FOnlineSessionIdHandle FOnlineIdRegistryRegistry::ToSessionId(EOnlineServices Services, const TArray<uint8>& RepData) const
+{
+	FOnlineSessionIdHandle Handle;
+	if (IOnlineSessionIdRegistry* Registry = GetSessionIdRegistry(Services))
+	{
+		Handle = Registry->FromReplicationData(RepData);
+	}
+	return Handle;
+}
+
+IOnlineSessionIdRegistry* FOnlineIdRegistryRegistry::GetSessionIdRegistry(EOnlineServices OnlineServices) const
+{
+	IOnlineSessionIdRegistry* Registry = nullptr;
+	if (const FSessionIdRegistryAndPriority* Found = SessionIdRegistries.Find(OnlineServices))
+	{
+		Registry = Found->Registry;
+	}
+	return Registry;
+}
+
+void FOnlineIdRegistryRegistry::RegisterSessionInviteIdRegistry(EOnlineServices OnlineServices, IOnlineSessionInviteIdRegistry* Registry, int32 Priority)
+{
+	FSessionInviteIdRegistryAndPriority* Found = SessionInviteIdRegistries.Find(OnlineServices);
+	if (Found == nullptr || Found->Priority < Priority)
+	{
+		SessionInviteIdRegistries.Emplace(OnlineServices, FSessionInviteIdRegistryAndPriority(Registry, Priority));
+	}
+}
+
+void FOnlineIdRegistryRegistry::UnregisterSessionInviteIdRegistry(EOnlineServices OnlineServices, int32 Priority)
+{
+	FSessionInviteIdRegistryAndPriority* Found = SessionInviteIdRegistries.Find(OnlineServices);
+	if (Found != nullptr && Found->Priority == Priority)
+	{
+		SessionInviteIdRegistries.Remove(OnlineServices);
+	}
+}
+
+FString FOnlineIdRegistryRegistry::ToLogString(const FOnlineSessionInviteIdHandle& Handle) const
+{
+	FString Result;
+	if (IOnlineSessionInviteIdRegistry* Registry = GetSessionInviteIdRegistry(Handle.GetOnlineServicesType()))
+	{
+		Result = Registry->ToLogString(Handle);
+	}
+	return Result;
+}
+
+TArray<uint8> FOnlineIdRegistryRegistry::ToReplicationData(const FOnlineSessionInviteIdHandle& Handle) const
+{
+	TArray<uint8> RepData;
+	if (IOnlineSessionInviteIdRegistry* Registry = GetSessionInviteIdRegistry(Handle.GetOnlineServicesType()))
+	{
+		RepData = Registry->ToReplicationData(Handle);
+	}
+	return RepData;
+}
+
+FOnlineSessionInviteIdHandle FOnlineIdRegistryRegistry::ToSessionInviteId(EOnlineServices Services, const TArray<uint8>& RepData) const
+{
+	FOnlineSessionInviteIdHandle Handle;
+	if (IOnlineSessionInviteIdRegistry* Registry = GetSessionInviteIdRegistry(Services))
+	{
+		Handle = Registry->FromReplicationData(RepData);
+	}
+	return Handle;
+}
+
+IOnlineSessionInviteIdRegistry* FOnlineIdRegistryRegistry::GetSessionInviteIdRegistry(EOnlineServices OnlineServices) const
+{
+	IOnlineSessionInviteIdRegistry* Registry = nullptr;
+	if (const FSessionInviteIdRegistryAndPriority* Found = SessionInviteIdRegistries.Find(OnlineServices))
+	{
+		Registry = Found->Registry;
+	}
+	return Registry;
+}
+
 FString ToLogString(const FOnlineAccountIdHandle& Id)
 {
 	return FString::Printf(TEXT("%s:%d (%s)"), LexToString(Id.GetOnlineServicesType()), Id.GetHandle(), *FOnlineIdRegistryRegistry::Get().ToLogString(Id));
@@ -302,7 +418,12 @@ FString ToLogString(const FOnlineLobbyIdHandle& Id)
 
 FString ToLogString(const FOnlineSessionIdHandle& Id)
 {
-	return FString::Printf(TEXT("%s:%d"), LexToString(Id.GetOnlineServicesType()), Id.GetHandle());
+	return FString::Printf(TEXT("%s:%d (%s)"), LexToString(Id.GetOnlineServicesType()), Id.GetHandle(), *FOnlineIdRegistryRegistry::Get().ToLogString(Id));
+}
+
+FString ToLogString(const FOnlineSessionInviteIdHandle& Id)
+{
+	return FString::Printf(TEXT("%s:%d (%s)"), LexToString(Id.GetOnlineServicesType()), Id.GetHandle(), *FOnlineIdRegistryRegistry::Get().ToLogString(Id));
 }
 
 FString ToLogString(const FOnlineVerifiedAuthTicketIdHandle& Id)
diff --git a/Engine/Source/Runtime/CoreOnline/Public/Online/CoreOnline.h b/Engine/Source/Runtime/CoreOnline/Public/Online/CoreOnline.h
--- a/Engine/Source/Runtime/CoreOnline/Public/Online/CoreOnline.h
+++ b/Engine/Source/Runtime/CoreOnline/Public/Online/CoreOnline.h
@@ -316,6 +316,7 @@ using FOnlineVerifiedAuthSessionIdHandle = TOnlineIdHandle<OnlineIdHandleTags::F
 COREONLINE_API FString ToLogString(const FOnlineAccountIdHandle& Id);
 COREONLINE_API FString ToLogString(const FOnlineLobbyIdHandle& Id);
 COREONLINE_API FString ToLogString(const FOnlineSessionIdHandle& Id);
+COREONLINE_API FString ToLogString(const FOnlineSessionInviteIdHandle& Id);
 COREONLINE_API FString ToLogString(const FOnlineVerifiedAuthTicketIdHandle& Id);
 COREONLINE_API FString ToLogString(const FOnlineVerifiedAuthSessionIdHandle& Id);
 
@@ -330,6 +331,8 @@ template<typename IdType>
 class IOnlineIdRegistry
 {
 public:
+	virtual ~IOnlineIdRegistry() = default;
+
 	virtual FString ToLogString(const TOnlineIdHandle<IdType>& Handle) const = 0;
 	virtual TArray<uint8> ToReplicationData(const TOnlineIdHandle<IdType>& Handle) const = 0;
 	virtual TOnlineIdHandle<IdType> FromReplicationData(const TArray<uint8>& ReplicationData) = 0;
@@ -355,16 +358,16 @@ public:
 	COREONLINE_API static void TearDown();
 
 	/**
-	 * Register a registry for a given OnlineServices implementation and TOnlineIdHandle type
+	 * Register a registry for a given OnlineServices implementation and IOnlineAccountIdHandle type
 	 *
 	 * @param OnlineServices Services that the registry is for
-	 * @param Registry the registry of online ids
+	 * @param Registry the registry of online account ids
 	 * @param Priority Integer priority, allows an existing registry to be extended and registered with a higher priority so it is used instead
 	 */
 	COREONLINE_API void RegisterAccountIdRegistry(EOnlineServices OnlineServices, IOnlineAccountIdRegistry* Registry, int32 Priority = 0);
 
 	/**
-	 * Unregister a previously registered Id registry
+	 * Unregister a previously registered Account Id registry
 	 *
 	 * @param OnlineServices Services that the registry is for
 	 * @param Priority Integer priority, will be unregistered only if the priority matches the one that is registered
@@ -377,20 +380,77 @@ public:
 
 	COREONLINE_API IOnlineAccountIdRegistry* GetAccountIdRegistry(EOnlineServices OnlineServices) const;
 
+	/**
+	 * Register a registry for a given OnlineServices implementation and IOnlineSessionIdHandle type
+	 *
+	 * @param OnlineServices Services that the registry is for
+	 * @param Registry the registry of online session ids
+	 * @param Priority Integer priority, allows an existing registry to be extended and registered with a higher priority so it is used instead
+	 */
+	COREONLINE_API void RegisterSessionIdRegistry(EOnlineServices OnlineServices, IOnlineSessionIdRegistry* Registry, int32 Priority = 0);
+
+	/**
+	 * Unregister a previously registered Session Id registry
+	 *
+	 * @param OnlineServices Services that the registry is for
+	 * @param Priority Integer priority, will be unregistered only if the priority matches the one that is registered
+	 */
+	COREONLINE_API void UnregisterSessionIdRegistry(EOnlineServices OnlineServices, int32 Priority = 0);
+
+	COREONLINE_API FString ToLogString(const FOnlineSessionIdHandle& Handle) const;
+	COREONLINE_API TArray<uint8> ToReplicationData(const FOnlineSessionIdHandle& Handle) const;
+	COREONLINE_API FOnlineSessionIdHandle ToSessionId(EOnlineServices Services, const TArray<uint8>& RepData) const;
+
+	COREONLINE_API IOnlineSessionIdRegistry* GetSessionIdRegistry(EOnlineServices OnlineServices) const;
+
+	/**
+	 * Register a registry for a given OnlineServices implementation and IOnlineSessionInviteIdHandle type
+	 *
+	 * @param OnlineServices Services that the registry is for
+	 * @param Registry the registry of online session ids
+	 * @param Priority Integer priority, allows an existing registry to be extended and registered with a higher priority so it is used instead
+	 */
+	COREONLINE_API void RegisterSessionInviteIdRegistry(EOnlineServices OnlineServices, IOnlineSessionInviteIdRegistry* Registry, int32 Priority = 0);
+
+	/**
+	 * Unregister a previously registered Session Invite Id registry
+	 *
+	 * @param OnlineServices Services that the registry is for
+	 * @param Priority Integer priority, will be unregistered only if the priority matches the one that is registered
+	 */
+	COREONLINE_API void UnregisterSessionInviteIdRegistry(EOnlineServices OnlineServices, int32 Priority = 0);
+
+	COREONLINE_API FString ToLogString(const FOnlineSessionInviteIdHandle& Handle) const;
+	COREONLINE_API TArray<uint8> ToReplicationData(const FOnlineSessionInviteIdHandle& Handle) const;
+	COREONLINE_API FOnlineSessionInviteIdHandle ToSessionInviteId(EOnlineServices Services, const TArray<uint8>& RepData) const;
+
+	COREONLINE_API IOnlineSessionInviteIdRegistry* GetSessionInviteIdRegistry(EOnlineServices OnlineServices) const;
+
 private:
 
-	struct FAccountIdRegistryAndPriority
+	template<typename IdType>
+	struct FOnlineIdRegistryAndPriority
 	{
-		FAccountIdRegistryAndPriority(IOnlineAccountIdRegistry* InRegistry, int32 InPriority)
+		FOnlineIdRegistryAndPriority(IOnlineIdRegistry<IdType>* InRegistry, int32 InPriority)
 			: Registry(InRegistry), Priority(InPriority) {}
 
-		IOnlineAccountIdRegistry* Registry;
+		virtual ~FOnlineIdRegistryAndPriority() = default;
+
+		IOnlineIdRegistry<IdType>* Registry;
 		int32 Priority;
 	};
 
+	typedef FOnlineIdRegistryAndPriority<OnlineIdHandleTags::FAccount> FAccountIdRegistryAndPriority;
+	typedef FOnlineIdRegistryAndPriority<OnlineIdHandleTags::FSession> FSessionIdRegistryAndPriority;
+	typedef FOnlineIdRegistryAndPriority<OnlineIdHandleTags::FSessionInvite> FSessionInviteIdRegistryAndPriority;
+
 	TMap<EOnlineServices, FAccountIdRegistryAndPriority> AccountIdRegistries;
 	TUniquePtr<FOnlineForeignAccountIdRegistry> ForeignAccountIdRegistry;
 
+	TMap<EOnlineServices, FSessionIdRegistryAndPriority> SessionIdRegistries;
+
+	TMap<EOnlineServices, FSessionInviteIdRegistryAndPriority> SessionInviteIdRegistries;
+
 	friend FLazySingleton;
 
 PACKAGE_SCOPE:
