Added logic to transmit V1 session update events in V2 interface.

diff --git a/Engine/Plugins/Online/Experimental/OnlineServicesOSSAdapter/Source/Private/Online/SessionsOSSAdapter.cpp b/Engine/Plugins/Online/Experimental/OnlineServicesOSSAdapter/Source/Private/Online/SessionsOSSAdapter.cpp
--- a/Engine/Plugins/Online/Experimental/OnlineServicesOSSAdapter/Source/Private/Online/SessionsOSSAdapter.cpp
+++ b/Engine/Plugins/Online/Experimental/OnlineServicesOSSAdapter/Source/Private/Online/SessionsOSSAdapter.cpp
@@ -159,14 +159,113 @@ void FSessionsOSSAdapter::Initialize()
 			TResult<TSharedRef<const FSession>, FOnlineError>(BuildV2Session(&InviteResult.Session)),
 			EUISessionJoinRequestedSource::FromInvitation
 		};
+
+		SessionEvents.OnUISessionJoinRequested.Broadcast(Event);
+	}));
+
+	SessionsInterface->AddOnSessionParticipantRemovedDelegate_Handle(FOnSessionParticipantRemovedDelegate::CreateLambda([this](FName SessionName, const FUniqueNetId& TargetUniqueNetId)
+	{
+		// We won't transmit events for a session or member that doesn't exist
+		const TOnlineResult<FGetSessionByName> GetSessionByNameResult = GetSessionByName({ SessionName });
+		if (GetSessionByNameResult.IsOk())
+		{
+			const TSharedRef<const FSession>& FoundSession = GetSessionByNameResult.GetOkValue().Session;
+
+			const FOnlineServicesOSSAdapter& ServicesOSSAdapter = static_cast<FOnlineServicesOSSAdapter&>(Services);
+			const FUniqueNetIdRef TargetUniqueNetIdRef = TargetUniqueNetId.AsShared();
+			const FOnlineAccountIdHandle TargetAccountIdHandle = ServicesOSSAdapter.GetAccountIdRegistry().FindOrAddHandle(TargetUniqueNetIdRef);
+
+			FSessionSettingsUpdate SettingsUpdate;
+			SettingsUpdate.RemovedSessionMembers.Add(TargetAccountIdHandle);
+
+			const FSessionUpdated Event { FoundSession , SettingsUpdate };
+
+			SessionEvents.OnSessionUpdated.Broadcast(Event);
+		}
 	}));
 
-	// All delegates handling Settings or Participant changes in the V1 implementation will be handled with calls to UpdateSession, and the FSessionUpdated event
+	SessionsInterface->AddOnSessionParticipantsChangeDelegate_Handle(FOnSessionParticipantsChangeDelegate::CreateLambda([this](FName SessionName, const FUniqueNetId& TargetUniqueNetId, bool bJoined)
+	{
+		// We won't transmit events for a session or member that doesn't exist
+		const TOnlineResult<FGetSessionByName> GetSessionByNameResult = GetSessionByName({ SessionName });
+		if (GetSessionByNameResult.IsOk())
+		{
+			const TSharedRef<const FSession>& FoundSession = GetSessionByNameResult.GetOkValue().Session;
+
+			const FOnlineServicesOSSAdapter& ServicesOSSAdapter = static_cast<FOnlineServicesOSSAdapter&>(Services);
+			const FUniqueNetIdRef TargetUniqueNetIdRef = TargetUniqueNetId.AsShared();
+			const FOnlineAccountIdHandle TargetAccountIdHandle = ServicesOSSAdapter.GetAccountIdRegistry().FindOrAddHandle(TargetUniqueNetIdRef);
+
+			FSessionSettingsUpdate SettingsUpdate;
+			if (bJoined)
+			{
+				if (const FSessionMember* SessionMember = FoundSession->SessionSettings.SessionMembers.Find(TargetAccountIdHandle))
+				{
+					FSessionMemberUpdate SessionMemberUpdate;
+					SessionMemberUpdate.UpdatedMemberSettings.Append(SessionMember->MemberSettings);
+
+					SettingsUpdate.UpdatedSessionMembers.Emplace(TargetAccountIdHandle, SessionMemberUpdate);
+				}
+			}
+			else
+			{
+				SettingsUpdate.RemovedSessionMembers.Add(TargetAccountIdHandle);
+			}
+
+			const FSessionUpdated Event { FoundSession , SettingsUpdate };
+
+			SessionEvents.OnSessionUpdated.Broadcast(Event);
+		}
+	}));
+
+	SessionsInterface->AddOnSessionParticipantSettingsUpdatedDelegate_Handle(FOnSessionParticipantSettingsUpdatedDelegate::CreateLambda([this](FName SessionName, const FUniqueNetId& TargetUniqueNetId, const FOnlineSessionSettings& UpdatedSettings)
+	{
+		// We won't transmit events for a session or member that doesn't exist
+		const TOnlineResult<FGetSessionByName> GetSessionByNameResult = GetSessionByName({ SessionName });
+		if (GetSessionByNameResult.IsOk())
+		{
+			const TSharedRef<const FSession>& FoundSession = GetSessionByNameResult.GetOkValue().Session;
+
+			const FOnlineServicesOSSAdapter& ServicesOSSAdapter = static_cast<FOnlineServicesOSSAdapter&>(Services);
+			const FUniqueNetIdRef TargetUniqueNetIdRef = TargetUniqueNetId.AsShared();
+			const FOnlineAccountIdHandle TargetAccountIdHandle = ServicesOSSAdapter.GetAccountIdRegistry().FindOrAddHandle(TargetUniqueNetIdRef);
+
+			FSessionSettingsUpdate SettingsUpdate;
+			if (const FSessionMember* SessionMember = FoundSession->SessionSettings.SessionMembers.Find(TargetAccountIdHandle))
+			{
+				const ::FSessionSettings* MemberSettings = UpdatedSettings.MemberSettings.Find(TargetUniqueNetIdRef);
+				const FCustomSessionSettingsMap UpdatedMemberSettings = GetV2SessionSettings(*MemberSettings);
+
+				FSessionMemberUpdate SessionMemberUpdate;
+				for (const TPair<FName, FCustomSessionSetting>& Entry : UpdatedMemberSettings)
+				{
+					if (const FCustomSessionSetting* Setting = SessionMember->MemberSettings.Find(Entry.Key))
+					{
+						// TODO: First part of the comparison commented until the SchemaVariant changes are applied
+						if (/*Setting->Data != Entry.Value.Data ||*/ Setting->Visibility != Entry.Value.Visibility || Setting->ID != Entry.Value.ID)
+						{
+							SessionMemberUpdate.UpdatedMemberSettings.Add(Entry);
+						}
+					}
+				}
+
+				SettingsUpdate.UpdatedSessionMembers.Emplace(TargetAccountIdHandle, SessionMemberUpdate);
+			}
+
+			const FSessionUpdated Event{ FoundSession , SettingsUpdate };
+
+			SessionEvents.OnSessionUpdated.Broadcast(Event);
+		}
+	}));
 }
 
 void FSessionsOSSAdapter::Shutdown()
 {
 	SessionsInterface->ClearOnSessionInviteReceivedDelegates(this);
+	SessionsInterface->ClearOnSessionUserInviteAcceptedDelegates(this);
+	SessionsInterface->ClearOnSessionParticipantRemovedDelegates(this);
+	SessionsInterface->ClearOnSessionParticipantsChangeDelegates(this);
+	SessionsInterface->ClearOnSessionParticipantSettingsUpdatedDelegates(this);
 
 	Super::Shutdown();
 }
