OSSv2: Rename Handle, UserId to AccountId when use FAccountId.

diff --git a/Engine/Plugins/Online/Experimental/OnlineServicesOSSAdapter/Source/Private/Online/AchievementsOSSAdapter.cpp b/Engine/Plugins/Online/Experimental/OnlineServicesOSSAdapter/Source/Private/Online/AchievementsOSSAdapter.cpp
--- a/Engine/Plugins/Online/Experimental/OnlineServicesOSSAdapter/Source/Private/Online/AchievementsOSSAdapter.cpp
+++ b/Engine/Plugins/Online/Experimental/OnlineServicesOSSAdapter/Source/Private/Online/AchievementsOSSAdapter.cpp
@@ -32,10 +32,10 @@ void FAchievementsOSSAdapter::PostInitialize()
 
 	AchievementsInterface->AddOnAchievementUnlockedDelegate_Handle(FOnAchievementUnlockedDelegate::CreateLambda([this](const FUniqueNetId& LocalUserIdV1, const FString& AchievementId)
 	{
-		const FAccountId LocalUserIdV2 = Services.Get<FAuthOSSAdapter>()->GetAccountIdHandle(LocalUserIdV1.AsShared());
+		const FAccountId LocalUserIdV2 = Services.Get<FAuthOSSAdapter>()->GetAccountId(LocalUserIdV1.AsShared());
 
 		FAchievementStateUpdated AchievementStateUpdated;
-		AchievementStateUpdated.LocalUserId = LocalUserIdV2;
+		AchievementStateUpdated.LocalAccountId = LocalUserIdV2;
 		AchievementStateUpdated.AchievementIds = { AchievementId };
 		OnAchievementStateUpdatedEvent.Broadcast(MoveTemp(AchievementStateUpdated));
 	}));
@@ -61,7 +61,7 @@ TOnlineAsyncOpHandle<FQueryAchievementDefinitions> FAchievementsOSSAdapter::Quer
 		TPromise<void> Promise;
 		TFuture<void> Future = Promise.GetFuture();
 
-		const FUniqueNetIdPtr LocalUserId = Services.Get<FAuthOSSAdapter>()->GetUniqueNetId(Op.GetParams().LocalUserId);
+		const FUniqueNetIdPtr LocalUserId = Services.Get<FAuthOSSAdapter>()->GetUniqueNetId(Op.GetParams().LocalAccountId);
 		if (!LocalUserId)
 		{
 			Op.SetError(Errors::InvalidUser());
@@ -95,7 +95,7 @@ TOnlineAsyncOpHandle<FQueryAchievementDefinitions> FAchievementsOSSAdapter::Quer
 		TPromise<void> Promise;
 		TFuture<void> Future = Promise.GetFuture();
 
-		const FUniqueNetIdPtr LocalUserId = Services.Get<FAuthOSSAdapter>()->GetUniqueNetId(Op.GetParams().LocalUserId);
+		const FUniqueNetIdPtr LocalUserId = Services.Get<FAuthOSSAdapter>()->GetUniqueNetId(Op.GetParams().LocalAccountId);
 		if (!LocalUserId)
 		{
 			Op.SetError(Errors::InvalidUser());
@@ -118,7 +118,7 @@ TOnlineAsyncOpHandle<FQueryAchievementDefinitions> FAchievementsOSSAdapter::Quer
 	})
 	.Then([this](TOnlineAsyncOp<FQueryAchievementDefinitions>& Op)
 	{
-		const FUniqueNetIdPtr LocalUserId = Services.Get<FAuthOSSAdapter>()->GetUniqueNetId(Op.GetParams().LocalUserId);
+		const FUniqueNetIdPtr LocalUserId = Services.Get<FAuthOSSAdapter>()->GetUniqueNetId(Op.GetParams().LocalAccountId);
 		if (!LocalUserId)
 		{
 			Op.SetError(Errors::InvalidUser());
@@ -165,7 +165,7 @@ TOnlineAsyncOpHandle<FQueryAchievementDefinitions> FAchievementsOSSAdapter::Quer
 
 TOnlineResult<FGetAchievementIds> FAchievementsOSSAdapter::GetAchievementIds(FGetAchievementIds::Params&& Params)
 {
-	const FUniqueNetIdPtr LocalUserId = Services.Get<FAuthOSSAdapter>()->GetUniqueNetId(Params.LocalUserId);
+	const FUniqueNetIdPtr LocalUserId = Services.Get<FAuthOSSAdapter>()->GetUniqueNetId(Params.LocalAccountId);
 	if (!LocalUserId)
 	{
 		return TOnlineResult<FGetAchievementIds>(Errors::InvalidUser());
@@ -184,7 +184,7 @@ TOnlineResult<FGetAchievementIds> FAchievementsOSSAdapter::GetAchievementIds(FGe
 
 TOnlineResult<FGetAchievementDefinition> FAchievementsOSSAdapter::GetAchievementDefinition(FGetAchievementDefinition::Params&& Params)
 {
-	const FUniqueNetIdPtr LocalUserId = Services.Get<FAuthOSSAdapter>()->GetUniqueNetId(Params.LocalUserId);
+	const FUniqueNetIdPtr LocalUserId = Services.Get<FAuthOSSAdapter>()->GetUniqueNetId(Params.LocalAccountId);
 	if (!LocalUserId)
 	{
 		return TOnlineResult<FGetAchievementDefinition>(Errors::InvalidUser());
@@ -214,7 +214,7 @@ TOnlineAsyncOpHandle<FQueryAchievementStates> FAchievementsOSSAdapter::QueryAchi
 		TPromise<void> Promise;
 		TFuture<void> Future = Promise.GetFuture();
 
-		const FUniqueNetIdPtr LocalUserId = Services.Get<FAuthOSSAdapter>()->GetUniqueNetId(Op.GetParams().LocalUserId);
+		const FUniqueNetIdPtr LocalUserId = Services.Get<FAuthOSSAdapter>()->GetUniqueNetId(Op.GetParams().LocalAccountId);
 		if (!LocalUserId)
 		{
 			Op.SetError(Errors::InvalidUser());
@@ -246,7 +246,7 @@ TOnlineAsyncOpHandle<FQueryAchievementStates> FAchievementsOSSAdapter::QueryAchi
 	})
 	.Then([this](TOnlineAsyncOp<FQueryAchievementStates>& Op)
 	{
-		const FUniqueNetIdPtr LocalUserId = Services.Get<FAuthOSSAdapter>()->GetUniqueNetId(Op.GetParams().LocalUserId);
+		const FUniqueNetIdPtr LocalUserId = Services.Get<FAuthOSSAdapter>()->GetUniqueNetId(Op.GetParams().LocalAccountId);
 		if (!LocalUserId)
 		{
 			Op.SetError(Errors::InvalidUser());
@@ -261,7 +261,7 @@ TOnlineAsyncOpHandle<FQueryAchievementStates> FAchievementsOSSAdapter::QueryAchi
 			return;
 		}
 
-		FAchievementStateMap& AchievementStates = UserToAchievementStates.Emplace(Op.GetParams().LocalUserId);
+		FAchievementStateMap& AchievementStates = UserToAchievementStates.Emplace(Op.GetParams().LocalAccountId);
 
 		for (FOnlineAchievement& Achievement : Achievements)
 		{
@@ -280,13 +280,13 @@ TOnlineAsyncOpHandle<FQueryAchievementStates> FAchievementsOSSAdapter::QueryAchi
 
 TOnlineResult<FGetAchievementState> FAchievementsOSSAdapter::GetAchievementState(FGetAchievementState::Params&& Params) const
 {
-	const FUniqueNetIdPtr LocalUserId = Services.Get<FAuthOSSAdapter>()->GetUniqueNetId(Params.LocalUserId);
+	const FUniqueNetIdPtr LocalUserId = Services.Get<FAuthOSSAdapter>()->GetUniqueNetId(Params.LocalAccountId);
 	if (!LocalUserId)
 	{
 		return TOnlineResult<FGetAchievementState>(Errors::InvalidUser());
 	}
 
-	const FAchievementStateMap* AchievementStates = UserToAchievementStates.Find(Params.LocalUserId);
+	const FAchievementStateMap* AchievementStates = UserToAchievementStates.Find(Params.LocalAccountId);
 	if (!AchievementStates)
 	{
 		// Call QueryAchievementStates first
@@ -309,7 +309,7 @@ TOnlineResult<FDisplayAchievementUI> FAchievementsOSSAdapter::DisplayAchievement
 		return TOnlineResult<FDisplayAchievementUI>(Errors::NotImplemented());
 	}
 
-	const int LocalUserNum = Services.Get<FAuthOSSAdapter>()->GetLocalUserNum(Params.LocalUserId);
+	const int LocalUserNum = Services.Get<FAuthOSSAdapter>()->GetLocalUserNum(Params.LocalAccountId);
 	if (LocalUserNum == INDEX_NONE)
 	{
 		return TOnlineResult<FDisplayAchievementUI>(Errors::InvalidUser());
diff --git a/Engine/Plugins/Online/Experimental/OnlineServicesOSSAdapter/Source/Private/Online/AuthOSSAdapter.cpp b/Engine/Plugins/Online/Experimental/OnlineServicesOSSAdapter/Source/Private/Online/AuthOSSAdapter.cpp
--- a/Engine/Plugins/Online/Experimental/OnlineServicesOSSAdapter/Source/Private/Online/AuthOSSAdapter.cpp
+++ b/Engine/Plugins/Online/Experimental/OnlineServicesOSSAdapter/Source/Private/Online/AuthOSSAdapter.cpp
@@ -100,9 +100,9 @@ TSharedPtr<FAccountInfoOSSAdapter> FAccountInfoRegistryOSSAdapter::Find(FPlatfor
 	return StaticCastSharedPtr<FAccountInfoOSSAdapter>(Super::Find(PlatformUserId));
 } 
 
-TSharedPtr<FAccountInfoOSSAdapter> FAccountInfoRegistryOSSAdapter::Find(FAccountId AccountIdHandle) const
+TSharedPtr<FAccountInfoOSSAdapter> FAccountInfoRegistryOSSAdapter::Find(FAccountId AccountId) const
 {
-	return StaticCastSharedPtr<FAccountInfoOSSAdapter>(Super::Find(AccountIdHandle));
+	return StaticCastSharedPtr<FAccountInfoOSSAdapter>(Super::Find(AccountId));
 }
 
 void FAccountInfoRegistryOSSAdapter::Register(const TSharedRef<FAccountInfoOSSAdapter>& AccountInfoNULL)
@@ -306,7 +306,7 @@ TOnlineAsyncOpHandle<FAuthLogout> FAuthOSSAdapter::Logout(FAuthLogout::Params&&
 	{
 		const FAuthLogout::Params& Params = InAsyncOp.GetParams();
 
-		TSharedPtr<FAccountInfoOSSAdapter> AccountInfoOSSAdapter = AccountInfoRegistryOSSAdapter.Find(Params.LocalUserId);
+		TSharedPtr<FAccountInfoOSSAdapter> AccountInfoOSSAdapter = AccountInfoRegistryOSSAdapter.Find(Params.LocalAccountId);
 		if (!AccountInfoOSSAdapter)
 		{
 			InAsyncOp.SetError(Errors::InvalidUser());
@@ -373,7 +373,7 @@ TOnlineAsyncOpHandle<FAuthQueryExternalServerAuthTicket> FAuthOSSAdapter::QueryE
 		InAsyncOp.Data.Set<TSharedRef<FExternalServerAuthTicketTranslationTraits>>(ExternalServerAuthTicketTranslationTraitsKeyName, TranslationTraits.GetOkValue());
 
 		// Look up logged in user.
-		TSharedPtr<FAccountInfoOSSAdapter> AccountInfoOSSAdapter = AccountInfoRegistryOSSAdapter.Find(Params.LocalUserId);
+		TSharedPtr<FAccountInfoOSSAdapter> AccountInfoOSSAdapter = AccountInfoRegistryOSSAdapter.Find(Params.LocalAccountId);
 		if (!AccountInfoOSSAdapter)
 		{
 			InAsyncOp.SetError(Errors::InvalidUser());
@@ -417,7 +417,7 @@ TOnlineAsyncOpHandle<FAuthQueryExternalAuthToken> FAuthOSSAdapter::QueryExternal
 		InAsyncOp.Data.Set<const FExternalAuthTokenTranslationTraits*>(ExternalAuthTokenTranslationTraitsKeyName, TranslationTraits);
 
 		// Look up logged in user.
-		TSharedPtr<FAccountInfoOSSAdapter> AccountInfoOSSAdapter = AccountInfoRegistryOSSAdapter.Find(Params.LocalUserId);
+		TSharedPtr<FAccountInfoOSSAdapter> AccountInfoOSSAdapter = AccountInfoRegistryOSSAdapter.Find(Params.LocalAccountId);
 		if (!AccountInfoOSSAdapter)
 		{
 			InAsyncOp.SetError(Errors::InvalidUser());
@@ -516,19 +516,19 @@ TOnlineAsyncOpHandle<FAuthQueryExternalAuthToken> FAuthOSSAdapter::QueryExternal
 	return Op->GetHandle();
 }
 
-FUniqueNetIdPtr FAuthOSSAdapter::GetUniqueNetId(FAccountId AccountIdHandle) const
+FUniqueNetIdPtr FAuthOSSAdapter::GetUniqueNetId(FAccountId AccountId) const
 {
-	return GetOnlineServicesOSSAdapter().GetAccountIdRegistry().GetIdValue(AccountIdHandle);
+	return GetOnlineServicesOSSAdapter().GetAccountIdRegistry().GetIdValue(AccountId);
 }
 
-FAccountId FAuthOSSAdapter::GetAccountIdHandle(const FUniqueNetIdRef& UniqueNetId) const
+FAccountId FAuthOSSAdapter::GetAccountId(const FUniqueNetIdRef& UniqueNetId) const
 {
 	return GetOnlineServicesOSSAdapter().GetAccountIdRegistry().FindOrAddHandle(UniqueNetId);
 }
 
-int32 FAuthOSSAdapter::GetLocalUserNum(FAccountId AccountIdHandle) const
+int32 FAuthOSSAdapter::GetLocalUserNum(FAccountId AccountId) const
 {
-	TSharedPtr<FAccountInfoOSSAdapter> AccountInfoOSSAdapter = AccountInfoRegistryOSSAdapter.Find(AccountIdHandle);
+	TSharedPtr<FAccountInfoOSSAdapter> AccountInfoOSSAdapter = AccountInfoRegistryOSSAdapter.Find(AccountId);
 	return AccountInfoOSSAdapter ? AccountInfoOSSAdapter->LocalUserNum : INDEX_NONE;
 }
 
diff --git a/Engine/Plugins/Online/Experimental/OnlineServicesOSSAdapter/Source/Private/Online/CommerceOSSAdapter.cpp b/Engine/Plugins/Online/Experimental/OnlineServicesOSSAdapter/Source/Private/Online/CommerceOSSAdapter.cpp
--- a/Engine/Plugins/Online/Experimental/OnlineServicesOSSAdapter/Source/Private/Online/CommerceOSSAdapter.cpp
+++ b/Engine/Plugins/Online/Experimental/OnlineServicesOSSAdapter/Source/Private/Online/CommerceOSSAdapter.cpp
@@ -66,7 +66,7 @@ void FCommerceOSSAdapter::PostInitialize()
 	MakeMulticastAdapter(this, GetPurchaseInterface()->OnUnexpectedPurchaseReceiptDelegates, [this](const FUniqueNetId& User)
 	{
 		FCommerceOnPurchaseComplete Event;
-		Event.LocalUserId = Auth->GetAccountIdHandle(User.AsShared());
+		Event.LocalAccountId = Auth->GetAccountId(User.AsShared());
 		OnPurchaseCompletedEvent.Broadcast(MoveTemp(Event));
 
 		return false; // don't unbind
@@ -86,7 +86,7 @@ TOnlineAsyncOpHandle<FCommerceQueryOffers> FCommerceOSSAdapter::QueryOffers(FCom
 	Op->Then([this](TOnlineAsyncOp<FCommerceQueryOffers>& InAsyncOp)
 	{
 		IOnlineStoreV2Ptr StorePtr = GetStoreInterface();
-		FUniqueNetIdPtr NetId = Auth->GetUniqueNetId(InAsyncOp.GetParams().LocalUserId);
+		FUniqueNetIdPtr NetId = Auth->GetUniqueNetId(InAsyncOp.GetParams().LocalAccountId);
 		StorePtr->QueryOffersByFilter(*NetId, FOnlineStoreFilter(), *MakeDelegateAdapter(InAsyncOp, [this, WeakOp = InAsyncOp.AsWeak()](bool bWasSuccessful, const TArray<FUniqueOfferId>& OfferIds, const FString& Error)
 		{
 			if (bWasSuccessful)
@@ -113,7 +113,7 @@ TOnlineAsyncOpHandle<FCommerceQueryOffersById> FCommerceOSSAdapter::QueryOffersB
 	Op->Then([this](TOnlineAsyncOp<FCommerceQueryOffersById>& InAsyncOp)
 	{
 		IOnlineStoreV2Ptr StorePtr = GetStoreInterface();
-		FUniqueNetIdPtr NetId = Auth->GetUniqueNetId(InAsyncOp.GetParams().LocalUserId);
+		FUniqueNetIdPtr NetId = Auth->GetUniqueNetId(InAsyncOp.GetParams().LocalAccountId);
 		StorePtr->QueryOffersById(*NetId, InAsyncOp.GetParams().OfferIds, *MakeDelegateAdapter(InAsyncOp, [this, WeakOp = InAsyncOp.AsWeak()](bool bWasSuccessful, const TArray<FUniqueOfferId>& OfferIds, const FString& Error)
 		{
 			if (bWasSuccessful)
@@ -186,7 +186,7 @@ TOnlineAsyncOpHandle<FCommerceShowStoreUI> FCommerceOSSAdapter::ShowStoreUI(FCom
 			return;
 		}
 
-		GetExternalUIInterface()->ShowStoreUI(Auth->GetLocalUserNum(InAsyncOp.GetParams().LocalUserId), FShowStoreParams());
+		GetExternalUIInterface()->ShowStoreUI(Auth->GetLocalUserNum(InAsyncOp.GetParams().LocalAccountId), FShowStoreParams());
 		InAsyncOp.SetResult({});
 	})
 	.Enqueue(GetSerialQueue());
@@ -207,7 +207,7 @@ TOnlineAsyncOpHandle<FCommerceCheckout> FCommerceOSSAdapter::Checkout(FCommerceC
 			CheckoutRequest.AddPurchaseOffer(FOfferNamespace(), Offer.OfferId, Offer.Quantity);
 		}
 
-		GetPurchaseInterface()->Checkout(*Auth->GetUniqueNetId(InAsyncOp.GetParams().LocalUserId), CheckoutRequest, *MakeDelegateAdapter(InAsyncOp, [this, WeakOp = InAsyncOp.AsWeak()](const ::FOnlineError& Result, const TSharedRef<FPurchaseReceipt>& Receipt) mutable
+		GetPurchaseInterface()->Checkout(*Auth->GetUniqueNetId(InAsyncOp.GetParams().LocalAccountId), CheckoutRequest, *MakeDelegateAdapter(InAsyncOp, [this, WeakOp = InAsyncOp.AsWeak()](const ::FOnlineError& Result, const TSharedRef<FPurchaseReceipt>& Receipt) mutable
 		{
 			if (Result.WasSuccessful())
 			{
@@ -216,7 +216,7 @@ TOnlineAsyncOpHandle<FCommerceCheckout> FCommerceOSSAdapter::Checkout(FCommerceC
 				WeakOp.Pin()->SetResult(MoveTemp(OpResult));
 
 				FCommerceOnPurchaseComplete Event;
-				Event.LocalUserId = WeakOp.Pin()->GetParams().LocalUserId;
+				Event.LocalAccountId = WeakOp.Pin()->GetParams().LocalAccountId;
 				OnPurchaseCompletedEvent.Broadcast(MoveTemp(Event));
 			}
 			else
@@ -236,13 +236,13 @@ TOnlineAsyncOpHandle<FCommerceQueryTransactionEntitlements> FCommerceOSSAdapter:
 
 	Op->Then([this](TOnlineAsyncOp<FCommerceQueryTransactionEntitlements>& InAsyncOp)
 	{
-		GetPurchaseInterface()->QueryReceipts(*Auth->GetUniqueNetId(InAsyncOp.GetParams().LocalUserId), true, *MakeDelegateAdapter(InAsyncOp, [this, WeakOp = InAsyncOp.AsWeak()](const ::FOnlineError& Result) mutable
+		GetPurchaseInterface()->QueryReceipts(*Auth->GetUniqueNetId(InAsyncOp.GetParams().LocalAccountId), true, *MakeDelegateAdapter(InAsyncOp, [this, WeakOp = InAsyncOp.AsWeak()](const ::FOnlineError& Result) mutable
 		{
 			if (Result.WasSuccessful())
 			{
 				FCommerceQueryTransactionEntitlements::Result OpResult;
 				TArray<FPurchaseReceipt> OutReceipts;
-				GetPurchaseInterface()->GetReceipts(*Auth->GetUniqueNetId(WeakOp.Pin()->GetParams().LocalUserId), OutReceipts);
+				GetPurchaseInterface()->GetReceipts(*Auth->GetUniqueNetId(WeakOp.Pin()->GetParams().LocalAccountId), OutReceipts);
 
 				for (const FPurchaseReceipt& Receipt : OutReceipts)
 				{
@@ -275,7 +275,7 @@ TOnlineAsyncOpHandle<FCommerceQueryEntitlements> FCommerceOSSAdapter::QueryEntit
 	
 	Op->Then([this](TOnlineAsyncOp<FCommerceQueryEntitlements>& InAsyncOp)
 	{
-		GetPurchaseInterface()->QueryReceipts(*Auth->GetUniqueNetId(InAsyncOp.GetParams().LocalUserId), true, *MakeDelegateAdapter(InAsyncOp, [this, WeakOp = InAsyncOp.AsWeak()](const ::FOnlineError& Result) mutable
+		GetPurchaseInterface()->QueryReceipts(*Auth->GetUniqueNetId(InAsyncOp.GetParams().LocalAccountId), true, *MakeDelegateAdapter(InAsyncOp, [this, WeakOp = InAsyncOp.AsWeak()](const ::FOnlineError& Result) mutable
 		{
 			if (Result.WasSuccessful())
 			{
@@ -296,7 +296,7 @@ TOnlineResult<FCommerceGetEntitlements> FCommerceOSSAdapter::GetEntitlements(FCo
 {
 	FCommerceGetEntitlements::Result Result;
 	TArray<FPurchaseReceipt> OutReceipts;
-	GetPurchaseInterface()->GetReceipts(*Auth->GetUniqueNetId(Params.LocalUserId), OutReceipts);
+	GetPurchaseInterface()->GetReceipts(*Auth->GetUniqueNetId(Params.LocalAccountId), OutReceipts);
 
 	for (const FPurchaseReceipt& Receipt : OutReceipts)
 	{
diff --git a/Engine/Plugins/Online/Experimental/OnlineServicesOSSAdapter/Source/Private/Online/ExternalUIOSSAdapter.cpp b/Engine/Plugins/Online/Experimental/OnlineServicesOSSAdapter/Source/Private/Online/ExternalUIOSSAdapter.cpp
--- a/Engine/Plugins/Online/Experimental/OnlineServicesOSSAdapter/Source/Private/Online/ExternalUIOSSAdapter.cpp
+++ b/Engine/Plugins/Online/Experimental/OnlineServicesOSSAdapter/Source/Private/Online/ExternalUIOSSAdapter.cpp
@@ -66,10 +66,10 @@ TOnlineAsyncOpHandle<FExternalUIShowLoginUI> FExternalUIOSSAdapter::ShowLoginUI(
 							{
 								if (Error.WasSuccessful())
 								{
-									FAccountId Handle = static_cast<FOnlineServicesOSSAdapter&>(Services).GetAccountIdRegistry().FindOrAddHandle(UniqueId.ToSharedRef());
+									FAccountId AccountId = static_cast<FOnlineServicesOSSAdapter&>(Services).GetAccountIdRegistry().FindOrAddHandle(UniqueId.ToSharedRef());
 									FExternalUIShowLoginUI::Result Result = { MakeShared<FAccountInfo>() };
 									Result.AccountInfo->PlatformUserId = PinnedOp->GetParams().PlatformUserId;
-									Result.AccountInfo->AccountId = Handle;
+									Result.AccountInfo->AccountId = AccountId;
 									Result.AccountInfo->LoginStatus = ELoginStatus::LoggedIn;
 
 									PinnedOp->SetResult(MoveTemp(Result));
@@ -103,7 +103,7 @@ TOnlineAsyncOpHandle<FExternalUIShowFriendsUI> FExternalUIOSSAdapter::ShowFriend
 
 	if (!Op->IsReady())
 	{
-		int32 LocalUserIndex = Services.Get<FAuthOSSAdapter>()->GetLocalUserNum(Op->GetParams().LocalUserId);
+		int32 LocalUserIndex = Services.Get<FAuthOSSAdapter>()->GetLocalUserNum(Op->GetParams().LocalAccountId);
 
 		if (LocalUserIndex < 0 || LocalUserIndex >= MAX_LOCAL_PLAYERS)
 		{
diff --git a/Engine/Plugins/Online/Experimental/OnlineServicesOSSAdapter/Source/Private/Online/LeaderboardsOSSAdapter.cpp b/Engine/Plugins/Online/Experimental/OnlineServicesOSSAdapter/Source/Private/Online/LeaderboardsOSSAdapter.cpp
--- a/Engine/Plugins/Online/Experimental/OnlineServicesOSSAdapter/Source/Private/Online/LeaderboardsOSSAdapter.cpp
+++ b/Engine/Plugins/Online/Experimental/OnlineServicesOSSAdapter/Source/Private/Online/LeaderboardsOSSAdapter.cpp
@@ -27,7 +27,7 @@ TOnlineAsyncOpHandle<FReadEntriesForUsers> FLeaderboardsOSSAdapter::ReadEntriesF
 
 	Op->Then([this](TOnlineAsyncOp<FReadEntriesForUsers>& Op)
 	{
-		const FUniqueNetIdPtr LocalUniqueNetId = Auth->GetUniqueNetId(Op.GetParams().LocalUserId);
+		const FUniqueNetIdPtr LocalUniqueNetId = Auth->GetUniqueNetId(Op.GetParams().LocalAccountId);
 
 		if (!LocalUniqueNetId)
 		{
@@ -35,19 +35,19 @@ TOnlineAsyncOpHandle<FReadEntriesForUsers> FLeaderboardsOSSAdapter::ReadEntriesF
 			return MakeFulfilledPromise<void>().GetFuture();
 		}
 
-		if (Op.GetParams().UserIds.IsEmpty() || Op.GetParams().BoardName.IsEmpty())
+		if (Op.GetParams().AccountIds.IsEmpty() || Op.GetParams().BoardName.IsEmpty())
 		{
 			Op.SetError(Errors::InvalidParams());
 			return MakeFulfilledPromise<void>().GetFuture();
 		}
 
 		TArray<FUniqueNetIdRef> NetIds;
-		for (const FAccountId& UserId : Op.GetParams().UserIds)
+		for (const FAccountId& AccountId : Op.GetParams().AccountIds)
 		{
-			const FUniqueNetIdPtr UniqueNetId = Auth->GetUniqueNetId(UserId);
+			const FUniqueNetIdPtr UniqueNetId = Auth->GetUniqueNetId(AccountId);
 			if (!UniqueNetId)
 			{
-				UE_LOG(LogTemp, Warning, TEXT("Failed to get unique net id for user %s"), *ToLogString(UserId));
+				UE_LOG(LogTemp, Warning, TEXT("Failed to get unique net id for user %s"), *ToLogString(AccountId));
 				Op.SetError(Errors::InvalidUser());
 				return MakeFulfilledPromise<void>().GetFuture();
 			}
@@ -104,7 +104,7 @@ TOnlineAsyncOpHandle<FReadEntriesAroundRank> FLeaderboardsOSSAdapter::ReadEntrie
 
 	Op->Then([this](TOnlineAsyncOp<FReadEntriesAroundRank>& Op)
 	{
-		const FUniqueNetIdPtr LocalUserId = Auth->GetUniqueNetId(Op.GetParams().LocalUserId);
+		const FUniqueNetIdPtr LocalUserId = Auth->GetUniqueNetId(Op.GetParams().LocalAccountId);
 		if (!LocalUserId)
 		{
 			Op.SetError(Errors::InvalidUser());
@@ -166,14 +166,14 @@ TOnlineAsyncOpHandle<FReadEntriesAroundUser> FLeaderboardsOSSAdapter::ReadEntrie
 
 	Op->Then([this](TOnlineAsyncOp<FReadEntriesAroundUser>& Op)
 	{
-		const FUniqueNetIdPtr LocalUserId = Auth->GetUniqueNetId(Op.GetParams().LocalUserId);
+		const FUniqueNetIdPtr LocalUserId = Auth->GetUniqueNetId(Op.GetParams().LocalAccountId);
 		if (!LocalUserId)
 		{
 			Op.SetError(Errors::InvalidUser());
 			return MakeFulfilledPromise<void>().GetFuture();
 		}
 
-		const FUniqueNetIdPtr TargetUniqueNetId = Auth->GetUniqueNetId(Op.GetParams().UserId);
+		const FUniqueNetIdPtr TargetUniqueNetId = Auth->GetUniqueNetId(Op.GetParams().AccountId);
 		if (!TargetUniqueNetId)
 		{
 			Op.SetError(Errors::InvalidUser());
@@ -253,7 +253,7 @@ void FLeaderboardsOSSAdapter::ReadLeaderboardsResultFromV1ToV2(const FString& Bo
 	{
 		const FVariantData* VariantData = OnlineStatsRow.Columns.Find(*BoardName);
 		FLeaderboardEntry& LeaderboardEntry = OutEntries.Emplace_GetRef();
-		LeaderboardEntry.UserId = Auth->GetAccountIdHandle(OnlineStatsRow.PlayerId.ToSharedRef());
+		LeaderboardEntry.AccountId = Auth->GetAccountId(OnlineStatsRow.PlayerId.ToSharedRef());
 		LeaderboardEntry.Rank = OnlineStatsRow.Rank;
 		VariantData->GetValue(LeaderboardEntry.Score);
 	}
diff --git a/Engine/Plugins/Online/Experimental/OnlineServicesOSSAdapter/Source/Private/Online/PresenceOSSAdapter.cpp b/Engine/Plugins/Online/Experimental/OnlineServicesOSSAdapter/Source/Private/Online/PresenceOSSAdapter.cpp
--- a/Engine/Plugins/Online/Experimental/OnlineServicesOSSAdapter/Source/Private/Online/PresenceOSSAdapter.cpp
+++ b/Engine/Plugins/Online/Experimental/OnlineServicesOSSAdapter/Source/Private/Online/PresenceOSSAdapter.cpp
@@ -103,7 +103,7 @@ TSharedRef<FUserPresence> FPresenceOSSAdapter::PresenceV1toV2(FOnlineUserPresenc
 {
 	TSharedRef<FUserPresence> PresenceV2 = MakeShared<FUserPresence>();
 
-	PresenceV2->UserId = Auth->GetAccountIdHandle(Presence.SessionId.ToSharedRef());
+	PresenceV2->AccountId = Auth->GetAccountId(Presence.SessionId.ToSharedRef());
 	PresenceV2->Status = StatusV1toV2(Presence.Status.State);
 	PresenceV2->GameStatus = GameStatusV1toV2(Presence);
 	PresenceV2->Joinability = (Presence.bIsJoinable) ? EUserPresenceJoinability::Public : EUserPresenceJoinability::Private;
@@ -118,7 +118,7 @@ TSharedRef<FOnlineUserPresence> FPresenceOSSAdapter::PresenceV2toV1(FUserPresenc
 {
 	TSharedRef<FOnlineUserPresence> PresenceV1 = MakeShared<FOnlineUserPresence>();
 
-	PresenceV1->SessionId = Auth->GetUniqueNetId(Presence.UserId);
+	PresenceV1->SessionId = Auth->GetUniqueNetId(Presence.AccountId);
 	PresenceV1->bIsOnline = (Presence.Status == EUserPresenceStatus::Online) ? 1 : 0;
 	PresenceV1->bIsPlaying = (Presence.GameStatus != EUserPresenceGameStatus::Unknown) ? 1 : 0;
 	PresenceV1->bIsPlayingThisGame = (Presence.GameStatus == EUserPresenceGameStatus::PlayingThisGame) ? 1 : 0;
@@ -146,7 +146,7 @@ TOnlineAsyncOpHandle<FQueryPresence> FPresenceOSSAdapter::QueryPresence(FQueryPr
 
 	Op->Then([this](TOnlineAsyncOp<FQueryPresence>& Op, TPromise<TOnlineResult<FQueryPresence>>&& Result)
 	{
-		const FUniqueNetIdPtr UniqueNetId = Auth->GetUniqueNetId(Op.GetParams().LocalUserId);
+		const FUniqueNetIdPtr UniqueNetId = Auth->GetUniqueNetId(Op.GetParams().LocalAccountId);
 
 		if (UniqueNetId)
 		{
@@ -179,7 +179,7 @@ TOnlineAsyncOpHandle<FBatchQueryPresence> FPresenceOSSAdapter::BatchQueryPresenc
 
 	Op->Then([this](TOnlineAsyncOp<FBatchQueryPresence>& Op, TPromise<TOnlineResult<FBatchQueryPresence>>&& Result)
 	{
-		const FUniqueNetIdPtr LocalUserId = Auth->GetUniqueNetId(Op.GetParams().LocalUserId);
+		const FUniqueNetIdPtr LocalUserId = Auth->GetUniqueNetId(Op.GetParams().LocalAccountId);
 		if (!LocalUserId)
 		{
 			Op.SetError(Errors::InvalidUser());
@@ -187,9 +187,9 @@ TOnlineAsyncOpHandle<FBatchQueryPresence> FPresenceOSSAdapter::BatchQueryPresenc
 		}
 
 		TArray<FUniqueNetIdRef> NetIds;
-		for (const FAccountId& TargetUserId : Op.GetParams().TargetUserIds)
+		for (const FAccountId& TargetAccountId : Op.GetParams().TargetAccountIds)
 		{
-			const FUniqueNetIdPtr TargetUserNetId = Auth->GetUniqueNetId(TargetUserId);
+			const FUniqueNetIdPtr TargetUserNetId = Auth->GetUniqueNetId(TargetAccountId);
 			if (!TargetUserNetId)
 			{
 				Op.SetError(Errors::InvalidParams());
@@ -218,7 +218,7 @@ TOnlineAsyncOpHandle<FBatchQueryPresence> FPresenceOSSAdapter::BatchQueryPresenc
 
 TOnlineResult<FGetCachedPresence> FPresenceOSSAdapter::GetCachedPresence(FGetCachedPresence::Params&& Params)
 {
-	const FUniqueNetId& User = *Auth->GetUniqueNetId(Params.LocalUserId);
+	const FUniqueNetId& User = *Auth->GetUniqueNetId(Params.LocalAccountId);
 	TSharedPtr<FOnlineUserPresence> Presence = MakeShared<FOnlineUserPresence>();
 	EOnlineCachedResult::Type Result = PresenceInt->GetCachedPresence(User, Presence);
 	if (Result != EOnlineCachedResult::Success)
@@ -239,7 +239,7 @@ TOnlineAsyncOpHandle<FUpdatePresence> FPresenceOSSAdapter::UpdatePresence(FUpdat
 	Op->Then([this](TOnlineAsyncOp<FUpdatePresence>& Op, TPromise<TOnlineResult<FUpdatePresence>>&& Result)
 	{
 		TSharedRef<const FOnlineUserPresence> PresenceV1 = PresenceV2toV1(*Op.GetParams().Presence);
-		PresenceInt->SetPresence(*Auth->GetUniqueNetId(Op.GetParams().LocalUserId), PresenceV1->Status, *MakeDelegateAdapter(this, [this, ResultPromise = MoveTemp(Result)](const FUniqueNetId& UserId, const bool bWasSuccessful) mutable
+		PresenceInt->SetPresence(*Auth->GetUniqueNetId(Op.GetParams().LocalAccountId), PresenceV1->Status, *MakeDelegateAdapter(this, [this, ResultPromise = MoveTemp(Result)](const FUniqueNetId& UserId, const bool bWasSuccessful) mutable
 		{
 			if (bWasSuccessful)
 			{
@@ -283,7 +283,7 @@ TOnlineAsyncOpHandle<FPartialUpdatePresence> FPresenceOSSAdapter::PartialUpdateP
 			Parameters.Properties = NewProperties;
 		}
 
-		PresenceInt->SetPresence(*Auth->GetUniqueNetId(Op.GetParams().LocalUserId), MoveTemp(Parameters), *MakeDelegateAdapter(this, [this, ResultPromise = MoveTemp(Result)](const FUniqueNetId& UserId, const bool bWasSuccessful) mutable
+		PresenceInt->SetPresence(*Auth->GetUniqueNetId(Op.GetParams().LocalAccountId), MoveTemp(Parameters), *MakeDelegateAdapter(this, [this, ResultPromise = MoveTemp(Result)](const FUniqueNetId& UserId, const bool bWasSuccessful) mutable
 		{
 			if (bWasSuccessful)
 			{
diff --git a/Engine/Plugins/Online/Experimental/OnlineServicesOSSAdapter/Source/Private/Online/PrivilegesOSSAdapter.cpp b/Engine/Plugins/Online/Experimental/OnlineServicesOSSAdapter/Source/Private/Online/PrivilegesOSSAdapter.cpp
--- a/Engine/Plugins/Online/Experimental/OnlineServicesOSSAdapter/Source/Private/Online/PrivilegesOSSAdapter.cpp
+++ b/Engine/Plugins/Online/Experimental/OnlineServicesOSSAdapter/Source/Private/Online/PrivilegesOSSAdapter.cpp
@@ -27,7 +27,7 @@ TOnlineAsyncOpHandle<FQueryUserPrivilege> FPrivilegesOSSAdapter::QueryUserPrivil
 	if (!Op->IsReady())
 	{
 		FAuthOSSAdapter* Auth = Services.Get<FAuthOSSAdapter>();
-		FUniqueNetIdPtr UniqueNetId = Auth->GetUniqueNetId(Op->GetParams().LocalUserId);
+		FUniqueNetIdPtr UniqueNetId = Auth->GetUniqueNetId(Op->GetParams().LocalAccountId);
 		if (!UniqueNetId.IsValid())
 		{
 			Op->SetError(Errors::InvalidUser());
diff --git a/Engine/Plugins/Online/Experimental/OnlineServicesOSSAdapter/Source/Private/Online/SessionsOSSAdapter.cpp b/Engine/Plugins/Online/Experimental/OnlineServicesOSSAdapter/Source/Private/Online/SessionsOSSAdapter.cpp
--- a/Engine/Plugins/Online/Experimental/OnlineServicesOSSAdapter/Source/Private/Online/SessionsOSSAdapter.cpp
+++ b/Engine/Plugins/Online/Experimental/OnlineServicesOSSAdapter/Source/Private/Online/SessionsOSSAdapter.cpp
@@ -317,16 +317,16 @@ void FSessionsOSSAdapter::Initialize()
 
 	SessionsInterface->AddOnSessionParticipantRemovedDelegate_Handle(FOnSessionParticipantRemovedDelegate::CreateLambda([this](FName SessionName, const FUniqueNetId& TargetUniqueNetId)
 	{
-		const FOnlineServicesOSSAdapter& ServicesOSSAdapter = static_cast<FOnlineServicesOSSAdapter&>(Services);
-		const FUniqueNetIdRef TargetUniqueNetIdRef = TargetUniqueNetId.AsShared();
-		const FAccountId TargetAccountIdHandle = ServicesOSSAdapter.GetAccountIdRegistry().FindOrAddHandle(TargetUniqueNetIdRef);
+			const FOnlineServicesOSSAdapter& ServicesOSSAdapter = static_cast<FOnlineServicesOSSAdapter&>(Services);
+			const FUniqueNetIdRef TargetUniqueNetIdRef = TargetUniqueNetId.AsShared();
+			const FAccountId TargetAccountId = ServicesOSSAdapter.GetAccountIdRegistry().FindOrAddHandle(TargetUniqueNetIdRef);
 
-		FSessionSettingsUpdate SettingsUpdate;
-		SettingsUpdate.RemovedSessionMembers.Add(TargetAccountIdHandle);
+			FSessionSettingsUpdate SettingsUpdate;
+			SettingsUpdate.RemovedSessionMembers.Add(TargetAccountId);
 
 		const FSessionUpdated Event { SessionName , SettingsUpdate };
 
-		SessionEvents.OnSessionUpdated.Broadcast(Event);
+			SessionEvents.OnSessionUpdated.Broadcast(Event);
 	}));
 
 	SessionsInterface->AddOnSessionParticipantsChangeDelegate_Handle(FOnSessionParticipantsChangeDelegate::CreateLambda([this](FName SessionName, const FUniqueNetId& TargetUniqueNetId, bool bJoined)
@@ -339,22 +339,22 @@ void FSessionsOSSAdapter::Initialize()
 
 			const FOnlineServicesOSSAdapter& ServicesOSSAdapter = static_cast<FOnlineServicesOSSAdapter&>(Services);
 			const FUniqueNetIdRef TargetUniqueNetIdRef = TargetUniqueNetId.AsShared();
-			const FAccountId TargetAccountIdHandle = ServicesOSSAdapter.GetAccountIdRegistry().FindOrAddHandle(TargetUniqueNetIdRef);
+			const FAccountId TargetAccountId = ServicesOSSAdapter.GetAccountIdRegistry().FindOrAddHandle(TargetUniqueNetIdRef);
 
 			FSessionSettingsUpdate SettingsUpdate;
 			if (bJoined)
 			{
-				if (const FSessionMember* SessionMember = FoundSession->SessionSettings.SessionMembers.Find(TargetAccountIdHandle))
+				if (const FSessionMember* SessionMember = FoundSession->SessionSettings.SessionMembers.Find(TargetAccountId))
 				{
 					FSessionMemberUpdate SessionMemberUpdate;
 					SessionMemberUpdate.UpdatedMemberSettings.Append(SessionMember->MemberSettings);
 
-					SettingsUpdate.UpdatedSessionMembers.Emplace(TargetAccountIdHandle, SessionMemberUpdate);
+					SettingsUpdate.UpdatedSessionMembers.Emplace(TargetAccountId, SessionMemberUpdate);
 				}
 			}
 			else
 			{
-				SettingsUpdate.RemovedSessionMembers.Add(TargetAccountIdHandle);
+				SettingsUpdate.RemovedSessionMembers.Add(TargetAccountId);
 			}
 
 			const FSessionUpdated Event { SessionName , SettingsUpdate };
@@ -373,10 +373,10 @@ void FSessionsOSSAdapter::Initialize()
 
 			const FOnlineServicesOSSAdapter& ServicesOSSAdapter = static_cast<FOnlineServicesOSSAdapter&>(Services);
 			const FUniqueNetIdRef TargetUniqueNetIdRef = TargetUniqueNetId.AsShared();
-			const FAccountId TargetAccountIdHandle = ServicesOSSAdapter.GetAccountIdRegistry().FindOrAddHandle(TargetUniqueNetIdRef);
+			const FAccountId TargetAccountId = ServicesOSSAdapter.GetAccountIdRegistry().FindOrAddHandle(TargetUniqueNetIdRef);
 
 			FSessionSettingsUpdate SettingsUpdate;
-			if (const FSessionMember* SessionMember = FoundSession->SessionSettings.SessionMembers.Find(TargetAccountIdHandle))
+			if (const FSessionMember* SessionMember = FoundSession->SessionSettings.SessionMembers.Find(TargetAccountId))
 			{
 				const ::FSessionSettings* MemberSettings = UpdatedSettings.MemberSettings.Find(TargetUniqueNetIdRef);
 				const FCustomSessionSettingsMap UpdatedMemberSettings = GetV2SessionSettings(*MemberSettings);
@@ -393,7 +393,7 @@ void FSessionsOSSAdapter::Initialize()
 					}
 				}
 
-				SettingsUpdate.UpdatedSessionMembers.Emplace(TargetAccountIdHandle, SessionMemberUpdate);
+				SettingsUpdate.UpdatedSessionMembers.Emplace(TargetAccountId, SessionMemberUpdate);
 			}
 
 			const FSessionUpdated Event{ SessionName , SettingsUpdate };
@@ -437,8 +437,8 @@ TOnlineAsyncOpHandle<FCreateSession> FSessionsOSSAdapter::CreateSession(FCreateS
 			return;
 		}
 
-		// Dedicated servers can update sessions so we won't check if LocalUserId is valid
-		const FUniqueNetIdRef LocalUserId = Services.Get<FAuthOSSAdapter>()->GetUniqueNetId(OpParams.LocalUserId).ToSharedRef();
+		// Dedicated servers can update sessions so we won't check if LocalAccountId is valid
+		const FUniqueNetIdRef LocalUserId = Services.Get<FAuthOSSAdapter>()->GetUniqueNetId(OpParams.LocalAccountId).ToSharedRef();
 
 		MakeMulticastAdapter(this, SessionsInterface->OnCreateSessionCompleteDelegates,
 		[this, WeakOp = Op.AsWeak()](FName SessionName, const bool bWasSuccessful) mutable
@@ -461,7 +461,7 @@ TOnlineAsyncOpHandle<FCreateSession> FSessionsOSSAdapter::CreateSession(FCreateS
 
 				LocalSessionsByName.Emplace(OpParams.SessionName, V2Session->SessionId);
 
-				NamedSessionUserMap.FindOrAdd(OpParams.LocalUserId).AddUnique(OpParams.SessionName);
+				NamedSessionUserMap.FindOrAdd(OpParams.LocalAccountId).AddUnique(OpParams.SessionName);
 
 				Op->SetResult({ });
 			}
@@ -579,11 +579,11 @@ TOnlineAsyncOpHandle<FLeaveSession> FSessionsOSSAdapter::LeaveSession(FLeaveSess
 				{
 					TSharedRef<const FSession> FoundSession = GetSessionByNameResult.GetOkValue().Session;
 
-					NamedSessionUserMap.FindChecked(OpParams.LocalUserId).Remove(OpParams.SessionName);
+					NamedSessionUserMap.FindChecked(OpParams.LocalAccountId).Remove(OpParams.SessionName);
 
 					if (FoundSession->SessionSettings.bPresenceEnabled)
 					{
-						PresenceSessionsUserMap.Remove(OpParams.LocalUserId);
+						PresenceSessionsUserMap.Remove(OpParams.LocalAccountId);
 					}
 
 					ClearSessionByName(OpParams.SessionName);
@@ -593,7 +593,7 @@ TOnlineAsyncOpHandle<FLeaveSession> FSessionsOSSAdapter::LeaveSession(FLeaveSess
 				Op->SetResult({ });
 
 				FSessionLeft SessionLeftEvent;
-				SessionLeftEvent.LocalUserIds.Append(Op->GetParams().LocalUsers);
+				SessionLeftEvent.LocalAccountIds.Append(Op->GetParams().LocalAccounts);
 
 				SessionEvents.OnSessionLeft.Broadcast(SessionLeftEvent);
 			}
@@ -633,14 +633,14 @@ TOnlineAsyncOpHandle<FFindSessions> FSessionsOSSAdapter::FindSessions(FFindSessi
 			return;
 		}
 
-		const FUniqueNetIdRef LocalUserId = Services.Get<FAuthOSSAdapter>()->GetUniqueNetId(OpParams.LocalUserId).ToSharedRef();
+		const FUniqueNetIdRef LocalUserId = Services.Get<FAuthOSSAdapter>()->GetUniqueNetId(OpParams.LocalAccountId).ToSharedRef();
 		if (!LocalUserId->IsValid())
 		{
 			Op.SetError(Errors::InvalidUser());
 			return;
 		}
 
-		if (PendingV1SessionSearchesPerUser.Contains(OpParams.LocalUserId))
+		if (PendingV1SessionSearchesPerUser.Contains(OpParams.LocalAccountId))
 		{
 			Op.SetError(Errors::AlreadyPending());
 			return;
@@ -660,11 +660,11 @@ TOnlineAsyncOpHandle<FFindSessions> FSessionsOSSAdapter::FindSessions(FFindSessi
 					{
 						if (bWasSuccessful)
 						{
-							TSharedRef<FOnlineSessionSearch>& PendingV1SessionSearch = PendingV1SessionSearchesPerUser.FindChecked(Op->GetParams().LocalUserId);
+							TSharedRef<FOnlineSessionSearch>& PendingV1SessionSearch = PendingV1SessionSearchesPerUser.FindChecked(Op->GetParams().LocalAccountId);
 
 							TArray<TSharedRef<FSession>> FoundSessions = BuildV2SessionSearchResults(PendingV1SessionSearch->SearchResults);
 
-							TArray<FOnlineSessionIdHandle> SearchResults = SearchResultsUserMap.FindOrAdd(Op->GetParams().LocalUserId);
+							TArray<FOnlineSessionIdHandle> SearchResults = SearchResultsUserMap.FindOrAdd(Op->GetParams().LocalAccountId);
 							for (const TSharedRef<FSession>& FoundSession : FoundSessions)
 							{
 								AllSessionsById.Emplace(FoundSession->SessionId, FoundSession);
@@ -678,7 +678,7 @@ TOnlineAsyncOpHandle<FFindSessions> FSessionsOSSAdapter::FindSessions(FFindSessi
 							Op->SetError(Errors::Unknown());
 						}
 
-						PendingV1SessionSearchesPerUser.Remove(Op->GetParams().LocalUserId);
+						PendingV1SessionSearchesPerUser.Remove(Op->GetParams().LocalAccountId);
 					}
 				});
 
@@ -690,7 +690,7 @@ TOnlineAsyncOpHandle<FFindSessions> FSessionsOSSAdapter::FindSessions(FFindSessi
 					TargetIdRef = ServicesOSSAdapter.GetAccountIdRegistry().GetIdValue(*OpParams.TargetUser).ToSharedRef();
 				}
 
-				PendingV1SessionSearchesPerUser.Emplace(OpParams.LocalUserId, MakeShared<FOnlineSessionSearch>());
+				PendingV1SessionSearchesPerUser.Emplace(OpParams.LocalAccountId, MakeShared<FOnlineSessionSearch>());
 
 				SessionsInterface->FindSessionById(*LocalUserId, *SessionIdRef, *TargetIdRef, *MakeDelegateAdapter(this, [this, WeakOp = Op.AsWeak()](int32 LocalUserNum, bool bWasSuccessful, const FOnlineSessionSearchResult& SingleSearchResult) mutable
 				{
@@ -702,7 +702,7 @@ TOnlineAsyncOpHandle<FFindSessions> FSessionsOSSAdapter::FindSessions(FFindSessi
 							return;
 						}
 
-						TSharedRef<FOnlineSessionSearch>& PendingV1SessionSearch = PendingV1SessionSearchesPerUser.FindChecked(Op->GetParams().LocalUserId);
+						TSharedRef<FOnlineSessionSearch>& PendingV1SessionSearch = PendingV1SessionSearchesPerUser.FindChecked(Op->GetParams().LocalAccountId);
 
 						PendingV1SessionSearch->SearchResults.Add(SingleSearchResult);
 					}
@@ -716,7 +716,7 @@ TOnlineAsyncOpHandle<FFindSessions> FSessionsOSSAdapter::FindSessions(FFindSessi
 		}
 		else if (OpParams.TargetUser.IsSet())
 		{
-			int32 LocalUserNum = Services.Get<FAuthOSSAdapter>()->GetLocalUserNum(OpParams.LocalUserId);
+			int32 LocalUserNum = Services.Get<FAuthOSSAdapter>()->GetLocalUserNum(OpParams.LocalAccountId);
 
 			MakeMulticastAdapter(this, SessionsInterface->OnFindFriendSessionCompleteDelegates[LocalUserNum],
 			[this, WeakOp = Op.AsWeak()](int32 LocalUserId, bool bWasSuccessful, const TArray<FOnlineSessionSearchResult>& FriendSearchResults)
@@ -726,7 +726,7 @@ TOnlineAsyncOpHandle<FFindSessions> FSessionsOSSAdapter::FindSessions(FFindSessi
 					if (bWasSuccessful)
 					{
 						TArray<TSharedRef<FSession>> FoundSessions = BuildV2SessionSearchResults(FriendSearchResults);
-						TArray<FOnlineSessionIdHandle> SearchResults = SearchResultsUserMap.FindOrAdd(Op->GetParams().LocalUserId);
+						TArray<FOnlineSessionIdHandle> SearchResults = SearchResultsUserMap.FindOrAdd(Op->GetParams().LocalAccountId);
 						for (const TSharedRef<FSession>& FoundSession : FoundSessions)
 						{
 							AllSessionsById.Emplace(FoundSession->SessionId, FoundSession);
@@ -740,7 +740,7 @@ TOnlineAsyncOpHandle<FFindSessions> FSessionsOSSAdapter::FindSessions(FFindSessi
 						Op->SetError(Errors::Unknown());						
 					}
 
-					PendingV1SessionSearchesPerUser.Remove(Op->GetParams().LocalUserId);
+					PendingV1SessionSearchesPerUser.Remove(Op->GetParams().LocalAccountId);
 				}
 			});
 
@@ -749,7 +749,7 @@ TOnlineAsyncOpHandle<FFindSessions> FSessionsOSSAdapter::FindSessions(FFindSessi
 			TArray<FUniqueNetIdRef> TargetUniqueNetIds = { ServicesOSSAdapter.GetAccountIdRegistry().GetIdValue(*OpParams.TargetUser).ToSharedRef() };
 
 			// The cached search information is not used in this mode, but we'll still update it to not have more than one search at the same time
-			PendingV1SessionSearchesPerUser.Emplace(OpParams.LocalUserId, MakeShared<FOnlineSessionSearch>());
+			PendingV1SessionSearchesPerUser.Emplace(OpParams.LocalAccountId, MakeShared<FOnlineSessionSearch>());
 
 			SessionsInterface->FindFriendSession(*LocalUserId, TargetUniqueNetIds);
 		}
@@ -762,10 +762,10 @@ TOnlineAsyncOpHandle<FFindSessions> FSessionsOSSAdapter::FindSessions(FFindSessi
 				{
 					if (bWasSuccessful)
 					{
-						TSharedRef<FOnlineSessionSearch>& PendingV1SessionSearch = PendingV1SessionSearchesPerUser.FindChecked(Op->GetParams().LocalUserId);
+						TSharedRef<FOnlineSessionSearch>& PendingV1SessionSearch = PendingV1SessionSearchesPerUser.FindChecked(Op->GetParams().LocalAccountId);
 
 						TArray<TSharedRef<FSession>> FoundSessions = BuildV2SessionSearchResults(PendingV1SessionSearch->SearchResults);
-						TArray<FOnlineSessionIdHandle> SearchResults = SearchResultsUserMap.FindOrAdd(Op->GetParams().LocalUserId);
+						TArray<FOnlineSessionIdHandle> SearchResults = SearchResultsUserMap.FindOrAdd(Op->GetParams().LocalAccountId);
 						for (const TSharedRef<FSession>& FoundSession : FoundSessions)
 						{
 							AllSessionsById.Emplace(FoundSession->SessionId, FoundSession);
@@ -779,11 +779,11 @@ TOnlineAsyncOpHandle<FFindSessions> FSessionsOSSAdapter::FindSessions(FFindSessi
 						Op->SetError(Errors::Unknown());						
 					}
 
-					PendingV1SessionSearchesPerUser.Remove(Op->GetParams().LocalUserId);
+					PendingV1SessionSearchesPerUser.Remove(Op->GetParams().LocalAccountId);
 				}
 			});
 
-			TSharedRef<FOnlineSessionSearch>& PendingV1SessionSearch = PendingV1SessionSearchesPerUser.Emplace(OpParams.LocalUserId, BuildV1SessionSearch(OpParams));
+			TSharedRef<FOnlineSessionSearch>& PendingV1SessionSearch = PendingV1SessionSearchesPerUser.Emplace(OpParams.LocalAccountId, BuildV1SessionSearch(OpParams));
 
 			SessionsInterface->FindSessions(*LocalUserId, PendingV1SessionSearch);
 		}
@@ -816,7 +816,7 @@ TOnlineAsyncOpHandle<FStartMatchmaking> FSessionsOSSAdapter::StartMatchmaking(FS
 			return;
 		}
 
-		if (PendingV1SessionSearchesPerUser.Contains(OpParams.LocalUserId))
+		if (PendingV1SessionSearchesPerUser.Contains(OpParams.LocalAccountId))
 		{
 			Op.SetError(Errors::AlreadyPending());
 			return;
@@ -825,7 +825,7 @@ TOnlineAsyncOpHandle<FStartMatchmaking> FSessionsOSSAdapter::StartMatchmaking(FS
 		FOnlineServicesOSSAdapter& ServicesOSSAdapter = static_cast<FOnlineServicesOSSAdapter&>(Services);
 
 		TArray<FSessionMatchmakingUser> MatchMakingUsers;
-		for (const TPair<FAccountId, FSessionMember>& LocalUser : OpParams.LocalUsers)
+		for (const TPair<FAccountId, FSessionMember>& LocalUser : OpParams.LocalAccounts)
 		{
 			FSessionMatchmakingUser NewMatchMakingUser { ServicesOSSAdapter.GetAccountIdRegistry().GetIdValue(LocalUser.Key).ToSharedRef() };
 
@@ -857,10 +857,10 @@ TOnlineAsyncOpHandle<FStartMatchmaking> FSessionsOSSAdapter::StartMatchmaking(FS
 		FFindSessions::Params FindSessionsParams;
 		FindSessionsParams.bFindLANSessions = false;
 		FindSessionsParams.Filters = OpParams.SearchFilters;
-		FindSessionsParams.LocalUserId = OpParams.LocalUserId;
+		FindSessionsParams.LocalAccountId = OpParams.LocalAccountId;
 		FindSessionsParams.MaxResults = 1;
 
-		TSharedRef<FOnlineSessionSearch>& SessionSearch = PendingV1SessionSearchesPerUser.Emplace(OpParams.LocalUserId, BuildV1SessionSearch(FindSessionsParams));
+		TSharedRef<FOnlineSessionSearch>& SessionSearch = PendingV1SessionSearchesPerUser.Emplace(OpParams.LocalAccountId, BuildV1SessionSearch(FindSessionsParams));
 
 		// QUESTION: StartMatchmaking causes both OnMatchmakingComplete and OnStartMatchmakingComplete to trigger.
 		// The first triggers the second so I used the latter, but the interface header did not specify if one should be used. Is this the correct one?
@@ -882,11 +882,11 @@ TOnlineAsyncOpHandle<FStartMatchmaking> FSessionsOSSAdapter::StartMatchmaking(FS
 
 				LocalSessionsByName.Emplace(OpParams.SessionName, V2Session->SessionId);
 
-				NamedSessionUserMap.FindOrAdd(OpParams.LocalUserId).AddUnique(OpParams.SessionName);
+				NamedSessionUserMap.FindOrAdd(OpParams.LocalAccountId).AddUnique(OpParams.SessionName);
 
 				Op->SetResult({ });
 
-				FSessionJoined Event{ { Op->GetParams().LocalUserId }, V2Session->SessionId };
+				FSessionJoined Event{ { Op->GetParams().LocalAccountId }, V2Session->SessionId };
 
 				SessionEvents.OnSessionJoined.Broadcast(Event);
 
@@ -922,7 +922,7 @@ TOnlineAsyncOpHandle<FJoinSession> FSessionsOSSAdapter::JoinSession(FJoinSession
 			return;
 		}
 
-		const FUniqueNetIdRef LocalUserId = Services.Get<FAuthOSSAdapter>()->GetUniqueNetId(OpParams.LocalUserId).ToSharedRef();
+		const FUniqueNetIdRef LocalUserId = Services.Get<FAuthOSSAdapter>()->GetUniqueNetId(OpParams.LocalAccountId).ToSharedRef();
 		if (!LocalUserId->IsValid())
 		{
 			Op.SetError(Errors::InvalidUser());
@@ -942,7 +942,7 @@ TOnlineAsyncOpHandle<FJoinSession> FSessionsOSSAdapter::JoinSession(FJoinSession
 
 				const FJoinSession::Params& OpParams = Op->GetParams();
 
-				TOnlineResult<FGetSessionById> GetSessionByIdResult = GetSessionById({ OpParams.LocalUserId, OpParams.SessionId });
+				TOnlineResult<FGetSessionById> GetSessionByIdResult = GetSessionById({ OpParams.LocalAccountId, OpParams.SessionId });
 				if (GetSessionByIdResult.IsError())
 				{
 					// If no result is found, the id might be expired, which we should notify
@@ -959,16 +959,16 @@ TOnlineAsyncOpHandle<FJoinSession> FSessionsOSSAdapter::JoinSession(FJoinSession
 
 				LocalSessionsByName.Emplace(OpParams.SessionName, OpParams.SessionId);
 
-				NamedSessionUserMap.FindOrAdd(OpParams.LocalUserId).AddUnique(OpParams.SessionName);
+				NamedSessionUserMap.FindOrAdd(OpParams.LocalAccountId).AddUnique(OpParams.SessionName);
 
 				if (FoundSession->SessionSettings.bPresenceEnabled)
 				{
-					FOnlineSessionIdHandle& PresenceSessionId = PresenceSessionsUserMap.FindOrAdd(OpParams.LocalUserId);
+					FOnlineSessionIdHandle& PresenceSessionId = PresenceSessionsUserMap.FindOrAdd(OpParams.LocalAccountId);
 					PresenceSessionId = FoundSession->SessionId;
 				}
 
 				// After successfully joining a session, we'll remove all related invites if any are found
-				if (TMap<FOnlineSessionInviteIdHandle, TSharedRef<FSessionInvite>>* UserMap = SessionInvitesUserMap.Find(OpParams.LocalUserId))
+				if (TMap<FOnlineSessionInviteIdHandle, TSharedRef<FSessionInvite>>* UserMap = SessionInvitesUserMap.Find(OpParams.LocalAccountId))
 				{
 					TArray<FOnlineSessionInviteIdHandle> InviteIdsToRemove;
 					for (const TPair<FOnlineSessionInviteIdHandle, TSharedRef<FSessionInvite>>& Entry : *UserMap)
@@ -986,13 +986,13 @@ TOnlineAsyncOpHandle<FJoinSession> FSessionsOSSAdapter::JoinSession(FJoinSession
 
 				Op->SetResult({ });
 
-				FSessionJoined Event{ { OpParams.LocalUserId }, FoundSession->SessionId };
+				FSessionJoined Event{ { OpParams.LocalAccountId }, FoundSession->SessionId };
 
 				SessionEvents.OnSessionJoined.Broadcast(Event);
 			}
 		});
 
-		TOnlineResult<FGetSessionById> GetSessionByIdResult = GetSessionById({ OpParams.LocalUserId, OpParams.SessionId });
+		TOnlineResult<FGetSessionById> GetSessionByIdResult = GetSessionById({ OpParams.LocalAccountId, OpParams.SessionId });
 		if (GetSessionByIdResult.IsError())
 		{
 			Op.SetError(MoveTemp(GetSessionByIdResult.GetErrorValue()));
@@ -1145,7 +1145,7 @@ TOnlineAsyncOpHandle<FSendSessionInvite> FSessionsOSSAdapter::SendSessionInvite(
 
 		FAuthOSSAdapter* Auth = Services.Get<FAuthOSSAdapter>();
 
-		const FUniqueNetIdRef LocalUserId = Auth->GetUniqueNetId(OpParams.LocalUserId).ToSharedRef();
+		const FUniqueNetIdRef LocalUserId = Auth->GetUniqueNetId(OpParams.LocalAccountId).ToSharedRef();
 		if (!LocalUserId->IsValid())
 		{
 			Op.SetError(Errors::InvalidUser());
@@ -1189,7 +1189,7 @@ TOnlineResult<FGetResolvedConnectString> FSessionsOSSAdapter::GetResolvedConnect
 {
 	if (Params.SessionId.IsValid())
 	{
-		TOnlineResult<FGetSessionById> Result = GetSessionById({ Params.LocalUserId, Params.SessionId });
+		TOnlineResult<FGetSessionById> Result = GetSessionById({ Params.LocalAccountId, Params.SessionId });
 		if (Result.IsOk())
 		{
 			FOnlineSessionSearchResult SearchResult;
@@ -1334,7 +1334,7 @@ TSharedRef<FSession> FSessionsOSSAdapter::BuildV2Session(const FOnlineSession* I
 
 	if (const FNamedOnlineSession* NamedInSession = static_cast<const FNamedOnlineSession*>(InSession))
 	{
-		Session->OwnerUserId = ServicesOSSAdapter.GetAccountIdRegistry().FindOrAddHandle(NamedInSession->LocalOwnerId->AsShared());
+		Session->OwnerAccountId = ServicesOSSAdapter.GetAccountIdRegistry().FindOrAddHandle(NamedInSession->LocalOwnerId->AsShared());
 		WriteV2SessionSettingsFromV1NamedSession(NamedInSession, Session->SessionSettings);
 	}
 
diff --git a/Engine/Plugins/Online/Experimental/OnlineServicesOSSAdapter/Source/Private/Online/StatsOSSAdapter.cpp b/Engine/Plugins/Online/Experimental/OnlineServicesOSSAdapter/Source/Private/Online/StatsOSSAdapter.cpp
--- a/Engine/Plugins/Online/Experimental/OnlineServicesOSSAdapter/Source/Private/Online/StatsOSSAdapter.cpp
+++ b/Engine/Plugins/Online/Experimental/OnlineServicesOSSAdapter/Source/Private/Online/StatsOSSAdapter.cpp
@@ -85,7 +85,7 @@ FOnlineError FStatsOSSAdapter::ConvertUpdateUsersStatsV2toV1(const TArray<FUserS
 {
 	for (const FUserStats& UserStats : UpdateUsersStats)
 	{
-		const FUniqueNetIdPtr UserId = Auth->GetUniqueNetId(UserStats.UserId);
+		const FUniqueNetIdPtr UserId = Auth->GetUniqueNetId(UserStats.AccountId);
 		if (!UserId)
 		{
 			return Errors::InvalidUser();
@@ -111,7 +111,7 @@ TOnlineAsyncOpHandle<FUpdateStats> FStatsOSSAdapter::UpdateStats(FUpdateStats::P
 
 	Op->Then([this](TOnlineAsyncOp<FUpdateStats>& Op)
 	{
-		const FUniqueNetIdPtr LocalUserId = Auth->GetUniqueNetId(Op.GetParams().LocalUserId);
+		const FUniqueNetIdPtr LocalUserId = Auth->GetUniqueNetId(Op.GetParams().LocalAccountId);
 		if (!LocalUserId)
 		{
 			Op.SetError(Errors::InvalidUser());
@@ -153,8 +153,8 @@ TOnlineAsyncOpHandle<FQueryStats> FStatsOSSAdapter::QueryStats(FQueryStats::Para
 
 	Op->Then([this](TOnlineAsyncOp<FQueryStats>& Op)
 	{
-		const FUniqueNetIdPtr LocalUniqueNetId = Auth->GetUniqueNetId(Op.GetParams().LocalUserId);
-		const FUniqueNetIdPtr TargetUniqueNetId = Auth->GetUniqueNetId(Op.GetParams().TargetUserId);
+		const FUniqueNetIdPtr LocalUniqueNetId = Auth->GetUniqueNetId(Op.GetParams().LocalAccountId);
+		const FUniqueNetIdPtr TargetUniqueNetId = Auth->GetUniqueNetId(Op.GetParams().TargetAccountId);
 
 		if (!LocalUniqueNetId || !TargetUniqueNetId)
 		{
@@ -181,7 +181,7 @@ TOnlineAsyncOpHandle<FQueryStats> FStatsOSSAdapter::QueryStats(FQueryStats::Para
 				}
 
 				FUserStats UserStats;
-				UserStats.UserId = PinnedOp->GetParams().TargetUserId;
+				UserStats.AccountId = PinnedOp->GetParams().TargetAccountId;
 				UserStats.Stats = Result.Stats;
 				CacheUserStats(UserStats);
 
@@ -200,7 +200,7 @@ TOnlineAsyncOpHandle<FBatchQueryStats> FStatsOSSAdapter::BatchQueryStats(FBatchQ
 
 	Op->Then([this](TOnlineAsyncOp<FBatchQueryStats>& Op)
 	{
-		const FUniqueNetIdPtr LocalUserId = Auth->GetUniqueNetId(Op.GetParams().LocalUserId);
+		const FUniqueNetIdPtr LocalUserId = Auth->GetUniqueNetId(Op.GetParams().LocalAccountId);
 		if (!LocalUserId)
 		{
 			Op.SetError(Errors::InvalidUser());
@@ -208,9 +208,9 @@ TOnlineAsyncOpHandle<FBatchQueryStats> FStatsOSSAdapter::BatchQueryStats(FBatchQ
 		}
 
 		TArray<FUniqueNetIdRef> NetIds;
-		for (const FAccountId& TargetUserId : Op.GetParams().TargetUserIds)
+		for (const FAccountId& TargetAccountId : Op.GetParams().TargetAccountIds)
 		{
-			const FUniqueNetIdPtr UniqueNetId = Auth->GetUniqueNetId(TargetUserId);
+			const FUniqueNetIdPtr UniqueNetId = Auth->GetUniqueNetId(TargetAccountId);
 			if (!UniqueNetId)
 			{
 				Op.SetError(Errors::InvalidUser());
@@ -234,7 +234,7 @@ TOnlineAsyncOpHandle<FBatchQueryStats> FStatsOSSAdapter::BatchQueryStats(FBatchQ
 				for (const TSharedRef<const FOnlineStatsUserStats>& OnlineStatsUserStats : UsersStatsResult)
 				{
 					FUserStats& UserStats = Result.UsersStats.Emplace_GetRef();
-					UserStats.UserId = Auth->GetAccountIdHandle(OnlineStatsUserStats->Account);
+					UserStats.AccountId = Auth->GetAccountId(OnlineStatsUserStats->Account);
 
 					for (const TPair<FString, FOnlineStatValue>& StatPair : OnlineStatsUserStats->Stats)
 					{
@@ -262,7 +262,7 @@ TOnlineAsyncOpHandle<FResetStats> FStatsOSSAdapter::ResetStats(FResetStats::Para
 
 	Op->Then([this](TOnlineAsyncOp<FResetStats>& Op)
 	{
-		const FUniqueNetIdPtr LocalUserId = Auth->GetUniqueNetId(Op.GetParams().LocalUserId);
+		const FUniqueNetIdPtr LocalUserId = Auth->GetUniqueNetId(Op.GetParams().LocalAccountId);
 		if (!LocalUserId)
 		{
 			Op.SetError(Errors::InvalidUser());
diff --git a/Engine/Plugins/Online/Experimental/OnlineServicesOSSAdapter/Source/Private/Online/UserFileOSSAdapter.cpp b/Engine/Plugins/Online/Experimental/OnlineServicesOSSAdapter/Source/Private/Online/UserFileOSSAdapter.cpp
--- a/Engine/Plugins/Online/Experimental/OnlineServicesOSSAdapter/Source/Private/Online/UserFileOSSAdapter.cpp
+++ b/Engine/Plugins/Online/Experimental/OnlineServicesOSSAdapter/Source/Private/Online/UserFileOSSAdapter.cpp
@@ -35,7 +35,7 @@ TOnlineAsyncOpHandle<FUserFileEnumerateFiles> FUserFileOSSAdapter::EnumerateFile
 
 	Op->Then([this](TOnlineAsyncOp<FUserFileEnumerateFiles>& Op)
 	{
-		const FUniqueNetIdPtr LocalUserId = Services.Get<FAuthOSSAdapter>()->GetUniqueNetId(Op.GetParams().LocalUserId);
+		const FUniqueNetIdPtr LocalUserId = Services.Get<FAuthOSSAdapter>()->GetUniqueNetId(Op.GetParams().LocalAccountId);
 		if (!LocalUserId)
 		{
 			Op.SetError(Errors::InvalidUser());
@@ -69,7 +69,7 @@ TOnlineAsyncOpHandle<FUserFileEnumerateFiles> FUserFileOSSAdapter::EnumerateFile
 	})
 	.Then([this](TOnlineAsyncOp<FUserFileEnumerateFiles>& Op)
 	{
-		const FUniqueNetIdPtr LocalUserId = Services.Get<FAuthOSSAdapter>()->GetUniqueNetId(Op.GetParams().LocalUserId);
+		const FUniqueNetIdPtr LocalUserId = Services.Get<FAuthOSSAdapter>()->GetUniqueNetId(Op.GetParams().LocalAccountId);
 		if (!LocalUserId)
 		{
 			Op.SetError(Errors::InvalidUser());
@@ -87,7 +87,7 @@ TOnlineAsyncOpHandle<FUserFileEnumerateFiles> FUserFileOSSAdapter::EnumerateFile
 			NewEnumeratedFiles.Emplace(File.FileName);
 		}
 
-		EnumeratedFiles.Emplace(Op.GetParams().LocalUserId, MoveTemp(NewEnumeratedFiles));
+		EnumeratedFiles.Emplace(Op.GetParams().LocalAccountId, MoveTemp(NewEnumeratedFiles));
 		Op.SetResult({});
 	})
 	.Enqueue(GetSerialQueue());
@@ -96,13 +96,13 @@ TOnlineAsyncOpHandle<FUserFileEnumerateFiles> FUserFileOSSAdapter::EnumerateFile
 
 TOnlineResult<FUserFileGetEnumeratedFiles> FUserFileOSSAdapter::GetEnumeratedFiles(FUserFileGetEnumeratedFiles::Params&& Params)
 {
-	const FUniqueNetIdPtr LocalUserId = Services.Get<FAuthOSSAdapter>()->GetUniqueNetId(Params.LocalUserId);
+	const FUniqueNetIdPtr LocalUserId = Services.Get<FAuthOSSAdapter>()->GetUniqueNetId(Params.LocalAccountId);
 	if (!LocalUserId)
 	{
 		return TOnlineResult<FUserFileGetEnumeratedFiles>(Errors::InvalidUser());
 	}
 
-	const TArray<FString>* UserFiles = EnumeratedFiles.Find(Params.LocalUserId);
+	const TArray<FString>* UserFiles = EnumeratedFiles.Find(Params.LocalAccountId);
 	if (!UserFiles)
 	{
 		// Call EnumerateFiles first
@@ -129,7 +129,7 @@ TOnlineAsyncOpHandle<FUserFileReadFile> FUserFileOSSAdapter::ReadFile(FUserFileR
 	{
 		const FUserFileReadFile::Params& Params = Op.GetParams();
 
-		const FUniqueNetIdPtr LocalUserId = Services.Get<FAuthOSSAdapter>()->GetUniqueNetId(Params.LocalUserId);
+		const FUniqueNetIdPtr LocalUserId = Services.Get<FAuthOSSAdapter>()->GetUniqueNetId(Params.LocalAccountId);
 		if (!LocalUserId)
 		{
 			Op.SetError(Errors::InvalidUser());
@@ -180,7 +180,7 @@ TOnlineAsyncOpHandle<FUserFileReadFile> FUserFileOSSAdapter::ReadFile(FUserFileR
 	{
 		const FUserFileReadFile::Params& Params = Op.GetParams();
 
-		const FUniqueNetIdPtr LocalUserId = Services.Get<FAuthOSSAdapter>()->GetUniqueNetId(Params.LocalUserId);
+		const FUniqueNetIdPtr LocalUserId = Services.Get<FAuthOSSAdapter>()->GetUniqueNetId(Params.LocalAccountId);
 		if (!LocalUserId)
 		{
 			Op.SetError(Errors::InvalidUser());
@@ -217,7 +217,7 @@ TOnlineAsyncOpHandle<FUserFileWriteFile> FUserFileOSSAdapter::WriteFile(FUserFil
 	{
 		const FUserFileWriteFile::Params& Params = Op.GetParams();
 
-		const FUniqueNetIdPtr LocalUserId = Services.Get<FAuthOSSAdapter>()->GetUniqueNetId(Params.LocalUserId);
+		const FUniqueNetIdPtr LocalUserId = Services.Get<FAuthOSSAdapter>()->GetUniqueNetId(Params.LocalAccountId);
 		if (!LocalUserId)
 		{
 			Op.SetError(Errors::InvalidUser());
@@ -281,7 +281,7 @@ TOnlineAsyncOpHandle<FUserFileDeleteFile> FUserFileOSSAdapter::DeleteFile(FUserF
 	{
 		const FUserFileDeleteFile::Params& Params = Op.GetParams();
 
-		const FUniqueNetIdPtr LocalUserId = Services.Get<FAuthOSSAdapter>()->GetUniqueNetId(Params.LocalUserId);
+		const FUniqueNetIdPtr LocalUserId = Services.Get<FAuthOSSAdapter>()->GetUniqueNetId(Params.LocalAccountId);
 		if (!LocalUserId)
 		{
 			Op.SetError(Errors::InvalidUser());
diff --git a/Engine/Plugins/Online/Experimental/OnlineServicesOSSAdapter/Source/Private/Online/UserInfoOSSAdapter.cpp b/Engine/Plugins/Online/Experimental/OnlineServicesOSSAdapter/Source/Private/Online/UserInfoOSSAdapter.cpp
--- a/Engine/Plugins/Online/Experimental/OnlineServicesOSSAdapter/Source/Private/Online/UserInfoOSSAdapter.cpp
+++ b/Engine/Plugins/Online/Experimental/OnlineServicesOSSAdapter/Source/Private/Online/UserInfoOSSAdapter.cpp
@@ -28,7 +28,7 @@ TOnlineAsyncOpHandle<FQueryUserInfo> FUserInfoOSSAdapter::QueryUserInfo(FQueryUs
 	TOnlineAsyncOpRef<FQueryUserInfo> Op = GetJoinableOp<FQueryUserInfo>(MoveTemp(InParams));
 	const FQueryUserInfo::Params& Params = Op->GetParams();
 
-	if (Params.UserIds.IsEmpty())
+	if (Params.AccountIds.IsEmpty())
 	{
 		Op->SetError(Errors::InvalidParams());
 		return Op->GetHandle();
@@ -39,7 +39,7 @@ TOnlineAsyncOpHandle<FQueryUserInfo> FUserInfoOSSAdapter::QueryUserInfo(FQueryUs
 		const FQueryUserInfo::Params& Params = Op.GetParams();
 		const FAuthOSSAdapter* AuthInterface = Services.Get<FAuthOSSAdapter>();
 
-		const int32 LocalUserNum = AuthInterface->GetLocalUserNum(Params.LocalUserId);
+		const int32 LocalUserNum = AuthInterface->GetLocalUserNum(Params.LocalAccountId);
 		if (LocalUserNum == INDEX_NONE)
 		{
 			Op.SetError(Errors::InvalidUser());
@@ -47,10 +47,10 @@ TOnlineAsyncOpHandle<FQueryUserInfo> FUserInfoOSSAdapter::QueryUserInfo(FQueryUs
 		}
 
 		TArray<FUniqueNetIdRef> UserIds;
-		UserIds.Reserve(Params.UserIds.Num());
-		for(const FAccountId& TargetAccountIdHandle : Params.UserIds)
+		UserIds.Reserve(Params.AccountIds.Num());
+		for(const FAccountId& TargetAccountId : Params.AccountIds)
 		{
-			FUniqueNetIdPtr TargetUserNetId = AuthInterface->GetUniqueNetId(TargetAccountIdHandle);
+			FUniqueNetIdPtr TargetUserNetId = AuthInterface->GetUniqueNetId(TargetAccountId);
 			if (!TargetUserNetId)
 			{
 				Op.SetError(Errors::InvalidParams());
@@ -97,13 +97,13 @@ TOnlineResult<FGetUserInfo> FUserInfoOSSAdapter::GetUserInfo(FGetUserInfo::Param
 {
 	const FAuthOSSAdapter* AuthInterface = Services.Get<FAuthOSSAdapter>();
 
-	const int32 LocalUserNum = AuthInterface->GetLocalUserNum(Params.LocalUserId);
+	const int32 LocalUserNum = AuthInterface->GetLocalUserNum(Params.LocalAccountId);
 	if (LocalUserNum == INDEX_NONE)
 	{
 		return TOnlineResult<FGetUserInfo>(Errors::InvalidUser());
 	}
 
-	const FUniqueNetIdPtr TargetUserNetId = AuthInterface->GetUniqueNetId(Params.UserId);
+	const FUniqueNetIdPtr TargetUserNetId = AuthInterface->GetUniqueNetId(Params.AccountId);
 	if (!TargetUserNetId)
 	{
 		return TOnlineResult<FGetUserInfo>(Errors::InvalidParams());
@@ -116,7 +116,7 @@ TOnlineResult<FGetUserInfo> FUserInfoOSSAdapter::GetUserInfo(FGetUserInfo::Param
 	}
 
 	TSharedRef<FUserInfo> UserInfo = MakeShared<FUserInfo>();
-	UserInfo->UserId = Params.UserId;
+	UserInfo->AccountId = Params.AccountId;
 	UserInfo->DisplayName = OnlineUser->GetDisplayName();
 
 	return TOnlineResult<FGetUserInfo>({ UserInfo });
diff --git a/Engine/Plugins/Online/Experimental/OnlineServicesOSSAdapter/Source/Public/Online/AuthOSSAdapter.h b/Engine/Plugins/Online/Experimental/OnlineServicesOSSAdapter/Source/Public/Online/AuthOSSAdapter.h
--- a/Engine/Plugins/Online/Experimental/OnlineServicesOSSAdapter/Source/Public/Online/AuthOSSAdapter.h
+++ b/Engine/Plugins/Online/Experimental/OnlineServicesOSSAdapter/Source/Public/Online/AuthOSSAdapter.h
@@ -44,7 +44,7 @@ public:
 	virtual ~FAccountInfoRegistryOSSAdapter() = default;
 
 	TSharedPtr<FAccountInfoOSSAdapter> Find(FPlatformUserId PlatformUserId) const;
-	TSharedPtr<FAccountInfoOSSAdapter> Find(FAccountId AccountIdHandle) const;
+	TSharedPtr<FAccountInfoOSSAdapter> Find(FAccountId AccountId) const;
 
 	void Register(const TSharedRef<FAccountInfoOSSAdapter>&UserAuthData);
 	void Unregister(FAccountId AccountId);
@@ -67,9 +67,9 @@ public:
 	virtual TOnlineAsyncOpHandle<FAuthQueryExternalServerAuthTicket> QueryExternalServerAuthTicket(FAuthQueryExternalServerAuthTicket::Params&& Params) override;
 	virtual TOnlineAsyncOpHandle<FAuthQueryExternalAuthToken> QueryExternalAuthToken(FAuthQueryExternalAuthToken::Params&& Params) override;
 
-	FUniqueNetIdPtr GetUniqueNetId(FAccountId AccountIdHandle) const;
-	FAccountId GetAccountIdHandle(const FUniqueNetIdRef& UniqueNetId) const;
-	int32 GetLocalUserNum(FAccountId AccountIdHandle) const;
+	FUniqueNetIdPtr GetUniqueNetId(FAccountId AccountId) const;
+	FAccountId GetAccountId(const FUniqueNetIdRef& UniqueNetId) const;
+	int32 GetLocalUserNum(FAccountId AccountId) const;
 
 protected:
 #if !UE_BUILD_SHIPPING
diff --git a/Engine/Plugins/Online/Experimental/OnlineServicesOSSAdapter/Source/Public/Online/PresenceOSSAdapter.h b/Engine/Plugins/Online/Experimental/OnlineServicesOSSAdapter/Source/Public/Online/PresenceOSSAdapter.h
--- a/Engine/Plugins/Online/Experimental/OnlineServicesOSSAdapter/Source/Public/Online/PresenceOSSAdapter.h
+++ b/Engine/Plugins/Online/Experimental/OnlineServicesOSSAdapter/Source/Public/Online/PresenceOSSAdapter.h
@@ -29,8 +29,8 @@ public:
 	virtual TOnlineAsyncOpHandle<FUpdatePresence> UpdatePresence(FUpdatePresence::Params&& Params) override;
 	virtual TOnlineAsyncOpHandle<FPartialUpdatePresence> PartialUpdatePresence(FPartialUpdatePresence::Params&& Params) override;
 
-	FUniqueNetIdRef GetUniqueNetId(FAccountId AccountIdHandle) const;
-	int32 GetLocalUserNum(FAccountId AccountIdHandle) const;
+	FUniqueNetIdRef GetUniqueNetId(FAccountId AccountId) const;
+	int32 GetLocalUserNum(FAccountId AccountId) const;
 
 protected:
 	const FAuthOSSAdapter* Auth;
diff --git a/Engine/Source/Runtime/CoreOnline/Private/Online/CoreOnline.cpp b/Engine/Source/Runtime/CoreOnline/Private/Online/CoreOnline.cpp
--- a/Engine/Source/Runtime/CoreOnline/Private/Online/CoreOnline.cpp
+++ b/Engine/Source/Runtime/CoreOnline/Private/Online/CoreOnline.cpp
@@ -237,46 +237,46 @@ void FOnlineIdRegistryRegistry::UnregisterAccountIdRegistry(EOnlineServices Onli
 	}
 }
 
-FString FOnlineIdRegistryRegistry::ToLogString(const FAccountId& Handle) const
+FString FOnlineIdRegistryRegistry::ToLogString(const FAccountId& AccountId) const
 {
 	FString Result;
-	if (IOnlineAccountIdRegistry* Registry = GetAccountIdRegistry(Handle.GetOnlineServicesType()))
+	if (IOnlineAccountIdRegistry* Registry = GetAccountIdRegistry(AccountId.GetOnlineServicesType()))
 	{
-		Result = Registry->ToLogString(Handle);
+		Result = Registry->ToLogString(AccountId);
 	}
 	else
 	{
-		Result = ForeignAccountIdRegistry->ToLogString(Handle);
+		Result = ForeignAccountIdRegistry->ToLogString(AccountId);
 	}
 	return Result;
 }
 
-TArray<uint8> FOnlineIdRegistryRegistry::ToReplicationData(const FAccountId& Handle) const
+TArray<uint8> FOnlineIdRegistryRegistry::ToReplicationData(const FAccountId& AccountId) const
 {
 	TArray<uint8> RepData;
-	if (IOnlineAccountIdRegistry* Registry = GetAccountIdRegistry(Handle.GetOnlineServicesType()))
+	if (IOnlineAccountIdRegistry* Registry = GetAccountIdRegistry(AccountId.GetOnlineServicesType()))
 	{
-		RepData = Registry->ToReplicationData(Handle);
+		RepData = Registry->ToReplicationData(AccountId);
 	}
 	else
 	{
-		RepData = ForeignAccountIdRegistry->ToReplicationData(Handle);
+		RepData = ForeignAccountIdRegistry->ToReplicationData(AccountId);
 	}
 	return RepData;
 }
 
 FAccountId FOnlineIdRegistryRegistry::ToAccountId(EOnlineServices Services, const TArray<uint8>& RepData) const
 {
-	FAccountId Handle;
+	FAccountId AccountId;
 	if (IOnlineAccountIdRegistry* Registry = GetAccountIdRegistry(Services))
 	{
-		Handle = Registry->FromReplicationData(RepData);
+		AccountId = Registry->FromReplicationData(RepData);
 	}
 	else
 	{
-		Handle = ForeignAccountIdRegistry->FromReplicationData(Services, RepData);
+		AccountId = ForeignAccountIdRegistry->FromReplicationData(Services, RepData);
 	}
-	return Handle;
+	return AccountId;
 }
 
 IOnlineAccountIdRegistry* FOnlineIdRegistryRegistry::GetAccountIdRegistry(EOnlineServices OnlineServices) const
diff --git a/Engine/Source/Runtime/CoreOnline/Private/Online/CoreOnlinePrivate.cpp b/Engine/Source/Runtime/CoreOnline/Private/Online/CoreOnlinePrivate.cpp
--- a/Engine/Source/Runtime/CoreOnline/Private/Online/CoreOnlinePrivate.cpp
+++ b/Engine/Source/Runtime/CoreOnline/Private/Online/CoreOnlinePrivate.cpp
@@ -4,13 +4,13 @@
 
 namespace UE::Online {
 
-FString FOnlineForeignAccountIdRegistry::ToLogString(const FAccountId& Handle) const
+FString FOnlineForeignAccountIdRegistry::ToLogString(const FAccountId& AccountId) const
 {
 	FString Result;
-	if (Handle.IsValid())
+	if (AccountId.IsValid())
 	{
-		const EOnlineServices HandleType = Handle.GetOnlineServicesType();
-		const uint32 HandleIndex = Handle.GetHandle() - 1;
+		const EOnlineServices HandleType = AccountId.GetOnlineServicesType();
+		const uint32 HandleIndex = AccountId.GetHandle() - 1;
 		if (const FRepData* RepDataForServices = OnlineServicesToRepData.Find(HandleType))
 		{
 			const TArray<TArray<uint8>>& RepDataArray = RepDataForServices->RepDataArray;
@@ -27,14 +27,14 @@ FString FOnlineForeignAccountIdRegistry::ToLogString(const FAccountId& Handle) c
 	return Result;
 }
 
-TArray<uint8> FOnlineForeignAccountIdRegistry::ToReplicationData(const FAccountId& Handle) const
+TArray<uint8> FOnlineForeignAccountIdRegistry::ToReplicationData(const FAccountId& AccountId) const
 {
 	TArray<uint8> RepData;
-	if (Handle.IsValid())
+	if (AccountId.IsValid())
 	{
-		const EOnlineServices HandleType = Handle.GetOnlineServicesType();
-		const uint32 HandleIndex = Handle.GetHandle() - 1;
-		if (const FRepData* RepDataForServices = OnlineServicesToRepData.Find(Handle.GetOnlineServicesType()))
+		const EOnlineServices HandleType = AccountId.GetOnlineServicesType();
+		const uint32 HandleIndex = AccountId.GetHandle() - 1;
+		if (const FRepData* RepDataForServices = OnlineServicesToRepData.Find(AccountId.GetOnlineServicesType()))
 		{
 			const TArray<TArray<uint8>>& RepDataArray = RepDataForServices->RepDataArray;
 			if (ensure(RepDataArray.IsValidIndex(HandleIndex)))
@@ -48,22 +48,22 @@ TArray<uint8> FOnlineForeignAccountIdRegistry::ToReplicationData(const FAccountI
 
 FAccountId FOnlineForeignAccountIdRegistry::FromReplicationData(EOnlineServices Services, const TArray<uint8>& RepData)
 {
-	FAccountId Handle;
+	FAccountId AccountId;
 	if (RepData.Num())
 	{
 		FRepData& RepDataForServices = OnlineServicesToRepData.FindOrAdd(Services);
 		if (FAccountId* Found = RepDataForServices.RepDataToHandle.Find(RepData))
 		{
-			Handle = *Found;
+			AccountId = *Found;
 		}
 		else
 		{
 			RepDataForServices.RepDataArray.Add(RepData);
-			Handle = FAccountId(Services, RepDataForServices.RepDataArray.Num());
-			RepDataForServices.RepDataToHandle.Emplace(RepData, Handle);
+			AccountId = FAccountId(Services, RepDataForServices.RepDataArray.Num());
+			RepDataForServices.RepDataToHandle.Emplace(RepData, AccountId);
 		}
 	}
-	return Handle;
+	return AccountId;
 }
 
 }
\ No newline at end of file
diff --git a/Engine/Source/Runtime/CoreOnline/Private/Online/CoreOnlinePrivate.h b/Engine/Source/Runtime/CoreOnline/Private/Online/CoreOnlinePrivate.h
--- a/Engine/Source/Runtime/CoreOnline/Private/Online/CoreOnlinePrivate.h
+++ b/Engine/Source/Runtime/CoreOnline/Private/Online/CoreOnlinePrivate.h
@@ -28,8 +28,8 @@ namespace UE::Online {
 class FOnlineForeignAccountIdRegistry
 {
 public:
-	FString ToLogString(const FAccountId& Handle) const;
-	TArray<uint8> ToReplicationData(const FAccountId& Handle) const;
+	FString ToLogString(const FAccountId& AccountId) const;
+	TArray<uint8> ToReplicationData(const FAccountId& AccountId) const;
 	FAccountId FromReplicationData(EOnlineServices Services, const TArray<uint8>& RepData);
 
 private:
diff --git a/Engine/Source/Runtime/CoreOnline/Private/Tests/CoreOnlineTests.cpp b/Engine/Source/Runtime/CoreOnline/Private/Tests/CoreOnlineTests.cpp
--- a/Engine/Source/Runtime/CoreOnline/Private/Tests/CoreOnlineTests.cpp
+++ b/Engine/Source/Runtime/CoreOnline/Private/Tests/CoreOnlineTests.cpp
@@ -14,8 +14,8 @@ IMPLEMENT_SIMPLE_AUTOMATION_TEST(FForeignIdRepTest,
 	const TArray<uint8> SourceRepData { 0, 1, 2, 3, 4, 5, 6, 7 };
 
 	FOnlineIdRegistryRegistry Registry;
-	const FAccountId Handle = Registry.ToAccountId(Services, SourceRepData);
-	const TArray<uint8> RepData = Registry.ToReplicationData(Handle);
+	const FAccountId AccountId = Registry.ToAccountId(Services, SourceRepData);
+	const TArray<uint8> RepData = Registry.ToReplicationData(AccountId);
 	UTEST_EQUAL(TEXT(""), RepData, SourceRepData);
 
 	return true;
diff --git a/Engine/Source/Runtime/CoreOnline/Public/Online/CoreOnline.h b/Engine/Source/Runtime/CoreOnline/Public/Online/CoreOnline.h
--- a/Engine/Source/Runtime/CoreOnline/Public/Online/CoreOnline.h
+++ b/Engine/Source/Runtime/CoreOnline/Public/Online/CoreOnline.h
@@ -374,8 +374,8 @@ public:
 	 */
 	COREONLINE_API void UnregisterAccountIdRegistry(EOnlineServices OnlineServices, int32 Priority = 0);
 
-	COREONLINE_API FString ToLogString(const FAccountId& Handle) const;
-	COREONLINE_API TArray<uint8> ToReplicationData(const FAccountId& Handle) const;
+	COREONLINE_API FString ToLogString(const FAccountId& AccountId) const;
+	COREONLINE_API TArray<uint8> ToReplicationData(const FAccountId& AccountId) const;
 	COREONLINE_API FAccountId ToAccountId(EOnlineServices Services, const TArray<uint8>& RepData) const;
 
 	COREONLINE_API IOnlineAccountIdRegistry* GetAccountIdRegistry(EOnlineServices OnlineServices) const;
@@ -487,9 +487,9 @@ struct FUniqueNetIdWrapper
 	{
 	}
 
-	FUniqueNetIdWrapper(const UE::Online::FAccountId& Handle)
+	FUniqueNetIdWrapper(const UE::Online::FAccountId& AccountId)
 	{
-		Variant.Emplace<UE::Online::FAccountId>(Handle);
+		Variant.Emplace<UE::Online::FAccountId>(AccountId);
 	}
 
 	// temporarily restored implicit conversion from FUniqueNetId
@@ -562,8 +562,8 @@ struct FUniqueNetIdWrapper
 		}
 		else
 		{
-			const UE::Online::FAccountId& Handle = GetV2();
-			return Handle.IsValid();
+			const UE::Online::FAccountId& AccountId = GetV2();
+			return AccountId.IsValid();
 		}
 	}
 
@@ -577,9 +577,9 @@ struct FUniqueNetIdWrapper
 		Variant.Emplace<FUniqueNetIdPtr>(InUniqueNetId);
 	}
 
-	virtual void SetAccountId(const UE::Online::FAccountId& Handle)
+	virtual void SetAccountId(const UE::Online::FAccountId& AccountId)
 	{
-		Variant.Emplace<UE::Online::FAccountId>(Handle);
+		Variant.Emplace<UE::Online::FAccountId>(AccountId);
 	}
 
 	/** @return unique id associated with this wrapper object */
diff --git a/Engine/Source/Runtime/Engine/Classes/GameFramework/OnlineReplStructs.h b/Engine/Source/Runtime/Engine/Classes/GameFramework/OnlineReplStructs.h
--- a/Engine/Source/Runtime/Engine/Classes/GameFramework/OnlineReplStructs.h
+++ b/Engine/Source/Runtime/Engine/Classes/GameFramework/OnlineReplStructs.h
@@ -62,10 +62,10 @@ struct FUniqueNetIdRepl : public FUniqueNetIdWrapper
 		FUniqueNetIdWrapper::SetUniqueNetId(UniqueNetId);
 	}
 
-	virtual void SetAccountId(const UE::Online::FAccountId& Handle) override
+	virtual void SetAccountId(const UE::Online::FAccountId& AccountId) override
 	{
 		ReplicationBytes.Empty();
-		FUniqueNetIdWrapper::SetAccountId(Handle);
+		FUniqueNetIdWrapper::SetAccountId(AccountId);
 	}
 
 	/** Export contents of this struct as a string */
diff --git a/Engine/Source/Runtime/Engine/Private/OnlineReplStructs.cpp b/Engine/Source/Runtime/Engine/Private/OnlineReplStructs.cpp
--- a/Engine/Source/Runtime/Engine/Private/OnlineReplStructs.cpp
+++ b/Engine/Source/Runtime/Engine/Private/OnlineReplStructs.cpp
@@ -227,9 +227,9 @@ void FUniqueNetIdRepl::MakeReplicationDataV2()
 	EUniqueIdEncodingFlags EncodingFlags = EUniqueIdEncodingFlags::IsEncoded;
 	EncodingFlags |= static_cast<EUniqueIdEncodingFlags>(TypeHash_V2 << 3);
 
-	UE::Online::FAccountId Handle = GetV2();
-	UE::Online::EOnlineServices OnlineServicesType = Handle.GetOnlineServicesType();
-	TArray<uint8> ReplicationData = UE::Online::FOnlineIdRegistryRegistry::Get().ToReplicationData(Handle);
+	UE::Online::FAccountId AccountId = GetV2();
+	UE::Online::EOnlineServices OnlineServicesType = AccountId.GetOnlineServicesType();
+	TArray<uint8> ReplicationData = UE::Online::FOnlineIdRegistryRegistry::Get().ToReplicationData(AccountId);
 	check(!ReplicationData.IsEmpty());
 
 	ReplicationBytes.Empty();
@@ -445,9 +445,9 @@ void FUniqueNetIdRepl::NetSerializeLoadV2(FArchive& Ar, const EUniqueIdEncodingF
 	TArray<uint8> ReplicationData;
 	Ar << ReplicationData;
 
-	const UE::Online::FAccountId Handle = UE::Online::FOnlineIdRegistryRegistry::Get().ToAccountId(OnlineServicesType, ReplicationData);
-	check(Handle.IsValid());
-	SetAccountId(Handle);
+	const UE::Online::FAccountId AccountId = UE::Online::FOnlineIdRegistryRegistry::Get().ToAccountId(OnlineServicesType, ReplicationData);
+	check(AccountId.IsValid());
+	SetAccountId(AccountId);
 }
 
 bool FUniqueNetIdRepl::Serialize(FArchive& Ar)
