Refactor of FSession into ISession and FSessionCommon
Missing Lyra fixes after recent changes

diff --git a/Engine/Plugins/Online/Experimental/OnlineServices/Source/OnlineServicesCommon/Private/Online/NboSerializerCommonSvc.cpp b/Engine/Plugins/Online/Experimental/OnlineServices/Source/OnlineServicesCommon/Private/Online/NboSerializerCommonSvc.cpp
--- a/Engine/Plugins/Online/Experimental/OnlineServices/Source/OnlineServicesCommon/Private/Online/NboSerializerCommonSvc.cpp
+++ b/Engine/Plugins/Online/Experimental/OnlineServices/Source/OnlineServicesCommon/Private/Online/NboSerializerCommonSvc.cpp
@@ -90,7 +90,7 @@ void SerializeToBuffer(FNboSerializeToBuffer& Packet, const FSessionSettings& Se
 	// SessionSettings.SessionMembers will be serialized in implementations as user id types will vary
 }
 
-void SerializeToBuffer(FNboSerializeToBuffer& Packet, const FSession& Session)
+void SerializeToBuffer(FNboSerializeToBuffer& Packet, const FSessionCommon& Session)
 {
 	// Session.OwnerUserId will be serialized in implementations as user id types will vary
 	// Session.SessionId will be serialized in implementations as user id types will vary
@@ -215,7 +215,7 @@ void SerializeFromBuffer(FNboSerializeFromBuffer& Packet, FSessionSettings& Sess
 	// SessionSettings.SessionMembers will be deserialized in implementations as user id types will vary
 }
 
-void SerializeFromBuffer(FNboSerializeFromBuffer& Packet, FSession& Session)
+void SerializeFromBuffer(FNboSerializeFromBuffer& Packet, FSessionCommon& Session)
 {
 	// Session.OwnerUserId will be deserialized in implementations as user id types will vary
 	// Session.SessionId will be deserialized in implementations as user id types will vary
diff --git a/Engine/Plugins/Online/Experimental/OnlineServices/Source/OnlineServicesCommon/Private/Online/SessionsCommon.cpp b/Engine/Plugins/Online/Experimental/OnlineServices/Source/OnlineServicesCommon/Private/Online/SessionsCommon.cpp
--- a/Engine/Plugins/Online/Experimental/OnlineServices/Source/OnlineServicesCommon/Private/Online/SessionsCommon.cpp
+++ b/Engine/Plugins/Online/Experimental/OnlineServices/Source/OnlineServicesCommon/Private/Online/SessionsCommon.cpp
@@ -41,7 +41,7 @@ namespace UE::Online {
 	{
 		FGetAllSessions::Result Result;
 
-		for (const TPair<FOnlineSessionIdHandle, TSharedRef<FSession>>& SessionPair : AllSessionsById)
+		for (const TPair<FOnlineSessionIdHandle, TSharedRef<FSessionCommon>>& SessionPair : AllSessionsById)
 		{
 			Result.Sessions.Add(SessionPair.Value);
 		}
@@ -313,7 +313,7 @@ namespace UE::Online {
 
 	TOnlineResult<FGetMutableSessionById> FSessionsCommon::GetMutableSessionById(FGetMutableSessionById::Params&& Params) const
 	{
-		if (const TSharedRef<FSession>* FoundSession = AllSessionsById.Find(Params.IdHandle))
+		if (const TSharedRef<FSessionCommon>* FoundSession = AllSessionsById.Find(Params.IdHandle))
 		{
 			return TOnlineResult<FGetMutableSessionById>({ *FoundSession });
 		}
@@ -378,7 +378,7 @@ namespace UE::Online {
 		{
 			return TOnlineResult<FAddSessionMembers>(GetMutableSessionByNameResult.GetErrorValue());
 		}
-		TSharedRef<FSession> FoundSession = GetMutableSessionByNameResult.GetOkValue().Session;
+		TSharedRef<FSessionCommon> FoundSession = GetMutableSessionByNameResult.GetOkValue().Session;
 
 		FSessionSettings& SessionSettings = FoundSession->SessionSettings;
 
@@ -448,7 +448,7 @@ namespace UE::Online {
 		{
 			return TOnlineResult<FRemoveSessionMembers>(GetMutableSessionByNameResult.GetErrorValue());
 		}
-		TSharedRef<FSession> FoundSession = GetMutableSessionByNameResult.GetOkValue().Session;
+		TSharedRef<FSessionCommon> FoundSession = GetMutableSessionByNameResult.GetOkValue().Session;
 
 		FSessionSettings& SessionSettings = FoundSession->SessionSettings;
 
@@ -495,7 +495,7 @@ namespace UE::Online {
 		{
 			return TOnlineResult<FRegisterPlayers>(GetMutableSessionByNameResult.GetErrorValue());
 		}
-		TSharedRef<FSession> FoundSession = GetMutableSessionByNameResult.GetOkValue().Session;
+		TSharedRef<FSessionCommon> FoundSession = GetMutableSessionByNameResult.GetOkValue().Session;
 
 		FSessionSettings& SessionSettings = FoundSession->SessionSettings;
 
@@ -539,7 +539,7 @@ namespace UE::Online {
 		{
 			return TOnlineResult<FUnregisterPlayers>(GetMutableSessionByNameResult.GetErrorValue());
 		}
-		TSharedRef<FSession> FoundSession = GetMutableSessionByNameResult.GetOkValue().Session;
+		TSharedRef<FSessionCommon> FoundSession = GetMutableSessionByNameResult.GetOkValue().Session;
 
 		FSessionSettings& SessionSettings = FoundSession->SessionSettings;
 
@@ -636,7 +636,7 @@ namespace UE::Online {
 				TOnlineResult<FGetSessionById> GetSessionByIdResult = GetSessionById({ Params.LocalAccountId, Entry.Value });
 				if (GetSessionByIdResult.IsOk())
 				{
-					if (GetSessionByIdResult.GetOkValue().Session->SessionSettings.bPresenceEnabled)
+					if (GetSessionByIdResult.GetOkValue().Session->GetSessionSettings().bPresenceEnabled)
 				{
 					UE_LOG(LogTemp, Warning, TEXT("[FSessionsCommon::CheckCreateSessionState] Could not create session with bPresenceEnabled set to true when another already exists [%s]."), *Entry.Key.ToString());
 
@@ -671,9 +671,9 @@ namespace UE::Online {
 		TOnlineResult<FGetSessionByName> GetSessionByNameResult = GetSessionByName({ Params.SessionName });
 		if (GetSessionByNameResult.IsOk())
 		{
-			TSharedRef<const FSession> FoundSession = GetSessionByNameResult.GetOkValue().Session;
+			TSharedRef<const ISession> FoundSession = GetSessionByNameResult.GetOkValue().Session;
 
-			if (!FoundSession->SessionSettings.bIsDedicatedServerSession)
+			if (!FoundSession->GetSessionSettings().bIsDedicatedServerSession)
 			{
 				if (!Services.GetAuthInterface()->IsLoggedIn(Params.LocalAccountId))
 				{
@@ -841,9 +841,9 @@ namespace UE::Online {
 			return GetSessionByIdResult.GetErrorValue();
 		}
 
-		TSharedRef<const FSession> FoundSession = GetSessionByIdResult.GetOkValue().Session;
+		TSharedRef<const ISession> FoundSession = GetSessionByIdResult.GetOkValue().Session;
 
-		const FSessionSettings& SessionSettings = FoundSession->SessionSettings;
+		const FSessionSettings& SessionSettings = FoundSession->GetSessionSettings();
 
 		for (const FAccountId& LocalAccountId : LocalAccountIds)
 		{
@@ -950,9 +950,9 @@ namespace UE::Online {
 		TOnlineResult<FGetSessionByName> GetSessionByNameResult = GetSessionByName({ Params.SessionName });
 		if (GetSessionByNameResult.IsOk())
 		{
-			TSharedRef<const FSession> FoundSession = GetSessionByNameResult.GetOkValue().Session;
+			TSharedRef<const ISession> FoundSession = GetSessionByNameResult.GetOkValue().Session;
 
-			if (!FoundSession->SessionSettings.bIsDedicatedServerSession)
+			if (!FoundSession->GetSessionSettings().bIsDedicatedServerSession)
 			{
 				if (!Services.GetAuthInterface()->IsLoggedIn(Params.LocalAccountId))
 				{
diff --git a/Engine/Plugins/Online/Experimental/OnlineServices/Source/OnlineServicesCommon/Public/Online/NboSerializerCommonSvc.h b/Engine/Plugins/Online/Experimental/OnlineServices/Source/OnlineServicesCommon/Public/Online/NboSerializerCommonSvc.h
--- a/Engine/Plugins/Online/Experimental/OnlineServices/Source/OnlineServicesCommon/Public/Online/NboSerializerCommonSvc.h
+++ b/Engine/Plugins/Online/Experimental/OnlineServices/Source/OnlineServicesCommon/Public/Online/NboSerializerCommonSvc.h
@@ -21,7 +21,7 @@ ONLINESERVICESCOMMON_API void SerializeToBuffer(FNboSerializeToBuffer& Packet, c
 ONLINESERVICESCOMMON_API void SerializeToBuffer(FNboSerializeToBuffer& Packet, const FCustomSessionSettingsMap& CustomSessionSettingsMap);
 ONLINESERVICESCOMMON_API void SerializeToBuffer(FNboSerializeToBuffer& Packet, const FSessionMember& SessionMember);
 ONLINESERVICESCOMMON_API void SerializeToBuffer(FNboSerializeToBuffer& Packet, const FSessionSettings& SessionSettings);
-ONLINESERVICESCOMMON_API void SerializeToBuffer(FNboSerializeToBuffer& Packet, const FSession& Session);
+ONLINESERVICESCOMMON_API void SerializeToBuffer(FNboSerializeToBuffer& Packet, const FSessionCommon& Session);
 
 /** SerializeFromBuffer methods */
 
@@ -31,7 +31,7 @@ ONLINESERVICESCOMMON_API void SerializeFromBuffer(FNboSerializeFromBuffer& Packe
 ONLINESERVICESCOMMON_API void SerializeFromBuffer(FNboSerializeFromBuffer& Packet, FCustomSessionSettingsMap& CustomSessionSettingsMap);
 ONLINESERVICESCOMMON_API void SerializeFromBuffer(FNboSerializeFromBuffer& Packet, FSessionMember& SessionMember);
 ONLINESERVICESCOMMON_API void SerializeFromBuffer(FNboSerializeFromBuffer& Packet, FSessionSettings& SessionSettings);
-ONLINESERVICESCOMMON_API void SerializeFromBuffer(FNboSerializeFromBuffer& Packet, FSession& Session);
+ONLINESERVICESCOMMON_API void SerializeFromBuffer(FNboSerializeFromBuffer& Packet, FSessionCommon& Session);
 
 /* NboSerializerCommonSvc */ }
 
diff --git a/Engine/Plugins/Online/Experimental/OnlineServices/Source/OnlineServicesCommon/Public/Online/SessionsCommon.h b/Engine/Plugins/Online/Experimental/OnlineServices/Source/OnlineServicesCommon/Public/Online/SessionsCommon.h
--- a/Engine/Plugins/Online/Experimental/OnlineServices/Source/OnlineServicesCommon/Public/Online/SessionsCommon.h
+++ b/Engine/Plugins/Online/Experimental/OnlineServices/Source/OnlineServicesCommon/Public/Online/SessionsCommon.h
@@ -127,6 +127,31 @@ public:
 	TOnlineBasicSessionInviteIdRegistry<FString> BasicRegistry;
 };
 
+class ONLINESERVICESCOMMON_API FSessionCommon : public ISession
+{
+public:
+	FSessionCommon() = default;
+	FSessionCommon(const FSessionCommon& InSession) = default;
+	virtual ~FSessionCommon() = default;
+
+	// ISession
+	virtual const FAccountId GetOwnerAccountId() const override				{ return OwnerAccountId; }
+	virtual const FOnlineSessionIdHandle GetSessionId() const override		{ return SessionId; }
+	virtual const FSessionSettings GetSessionSettings() const override		{ return SessionSettings; }
+
+	virtual FString ToLogString() const override							{ return FString(); } // TODO: Implement after completing refactor
+
+public:
+	/** The user who currently owns the session */
+	FAccountId OwnerAccountId;
+
+	/** The id for the session, platform dependent */
+	FOnlineSessionIdHandle SessionId;
+
+	/** Set of session properties that can be altered by the session owner */
+	FSessionSettings SessionSettings;
+};
+
 struct FGetMutableSessionByName
 {
 	static constexpr TCHAR Name[] = TEXT("GetMutableSessionByName");
@@ -138,7 +163,7 @@ struct FGetMutableSessionByName
 
 	struct Result
 	{
-		TSharedRef<FSession> Session;
+		TSharedRef<FSessionCommon> Session;
 
 		Result() = delete; // cannot default construct due to TSharedRef
 	};
@@ -155,7 +180,7 @@ struct FGetMutableSessionById
 
 	struct Result
 	{
-		TSharedRef<FSession> Session;
+		TSharedRef<FSessionCommon> Session;
 
 		Result() = delete; // cannot default construct due to TSharedRef
 	};
@@ -270,7 +295,7 @@ protected:
 	TMap<FAccountId, TSharedRef<TOnlineAsyncOp<FFindSessions>>> CurrentSessionSearchHandlesUserMap;
 
 	/** Set of every distinct FSession found, indexed by Id */
-	TMap<FOnlineSessionIdHandle, TSharedRef<FSession>> AllSessionsById;
+	TMap<FOnlineSessionIdHandle, TSharedRef<FSessionCommon>> AllSessionsById;
 };
 
 namespace Meta {
diff --git a/Engine/Plugins/Online/Experimental/OnlineServices/Source/OnlineServicesCommonEngineUtils/Private/Online/SessionsLAN.cpp b/Engine/Plugins/Online/Experimental/OnlineServices/Source/OnlineServicesCommonEngineUtils/Private/Online/SessionsLAN.cpp
--- a/Engine/Plugins/Online/Experimental/OnlineServices/Source/OnlineServicesCommonEngineUtils/Private/Online/SessionsLAN.cpp
+++ b/Engine/Plugins/Online/Experimental/OnlineServices/Source/OnlineServicesCommonEngineUtils/Private/Online/SessionsLAN.cpp
@@ -89,12 +89,12 @@ void FSessionLAN::Initialize()
 	SessionSettings.CustomSettings.Add(CONNECT_STRING_TAG, ConnectString);
 }
 
-FSessionLAN& FSessionLAN::Cast(FSession& InSession)
+FSessionLAN& FSessionLAN::Cast(FSessionCommon& InSession)
 {
 	return static_cast<FSessionLAN&>(InSession);
 }
 
-const FSessionLAN& FSessionLAN::Cast(const FSession& InSession)
+const FSessionLAN& FSessionLAN::Cast(const ISession& InSession)
 {
 	return static_cast<const FSessionLAN&>(InSession);
 }
@@ -199,7 +199,7 @@ TOnlineAsyncOpHandle<FUpdateSession> FSessionsLAN::UpdateSession(FUpdateSession:
 		TOnlineResult<FGetMutableSessionByName> GetMutableSessionByNameResult = GetMutableSessionByName({ OpParams.SessionName });
 		check(GetMutableSessionByNameResult.IsOk());
 		
-		TSharedRef<FSession> FoundSession = GetMutableSessionByNameResult.GetOkValue().Session;
+		TSharedRef<FSessionCommon> FoundSession = GetMutableSessionByNameResult.GetOkValue().Session;
 		FSessionSettings& SessionSettings = FoundSession->SessionSettings;
 		const FSessionSettingsUpdate& UpdatedSettings = OpParams.Mutations;
 
@@ -303,20 +303,17 @@ TOnlineAsyncOpHandle<FJoinSession> FSessionsLAN::JoinSession(FJoinSession::Param
 			return ;
 		}
 
-		const TSharedRef<const FSession>& FoundSession = GetSessionByIdResult.GetOkValue().Session;
-		TSharedRef<FSession> NewSessionLANRef = ConstCastSharedRef<FSession>(FoundSession);
+		const TSharedRef<const ISession>& FoundSession = GetSessionByIdResult.GetOkValue().Session;
 
-		// We save the local object for the session, and set up the appropriate references
-		AllSessionsById.Emplace(NewSessionLANRef->SessionId, NewSessionLANRef);
-
-		LocalSessionsByName.Emplace(OpParams.SessionName, NewSessionLANRef->SessionId);
+		// We set up the appropriate references for the session
+		LocalSessionsByName.Emplace(OpParams.SessionName, FoundSession->GetSessionId());
 
 		NamedSessionUserMap.FindOrAdd(OpParams.LocalAccountId).AddUnique(OpParams.SessionName);
 
-		if (NewSessionLANRef->SessionSettings.bPresenceEnabled)
+		if (FoundSession->GetSessionSettings().bPresenceEnabled)
 		{
 			FOnlineSessionIdHandle& PresenceSessionId = PresenceSessionsUserMap.FindOrAdd(OpParams.LocalAccountId);
-			PresenceSessionId = NewSessionLANRef->SessionId;
+			PresenceSessionId = FoundSession->GetSessionId();
 		}
 
 		Op.SetResult(FJoinSession::Result{ });
@@ -325,7 +322,7 @@ TOnlineAsyncOpHandle<FJoinSession> FSessionsLAN::JoinSession(FJoinSession::Param
 		LocalAccountIds.Reserve(OpParams.LocalAccounts.Num());
 		OpParams.LocalAccounts.GenerateKeyArray(LocalAccountIds);
 
-		FSessionJoined SessionJoinedEvent = { MoveTemp(LocalAccountIds), NewSessionLANRef->SessionId };
+		FSessionJoined SessionJoinedEvent = { MoveTemp(LocalAccountIds), FoundSession->GetSessionId() };
 		
 		SessionEvents.OnSessionJoined.Broadcast(SessionJoinedEvent);
 	})
@@ -352,17 +349,17 @@ TOnlineAsyncOpHandle<FLeaveSession> FSessionsLAN::LeaveSession(FLeaveSession::Pa
 		TOnlineResult<FGetSessionByName> GetSessionByNameResult = GetSessionByName({ OpParams.SessionName });
 		if (GetSessionByNameResult.IsOk())
 		{
-			TSharedRef<const FSession> FoundSession = GetSessionByNameResult.GetOkValue().Session;
+			TSharedRef<const ISession> FoundSession = GetSessionByNameResult.GetOkValue().Session;
 
 			NamedSessionUserMap.FindChecked(OpParams.LocalAccountId).Remove(OpParams.SessionName);
 
-			if (FoundSession->SessionSettings.bPresenceEnabled)
+			if (FoundSession->GetSessionSettings().bPresenceEnabled)
 			{
 				PresenceSessionsUserMap.Remove(OpParams.LocalAccountId);
 			}
 
 			ClearSessionByName(OpParams.SessionName);
-			ClearSessionById(FoundSession->SessionId);
+			ClearSessionById(FoundSession->GetSessionId());
 		}
 
 		Op.SetResult(FLeaveSession::Result{ });
@@ -464,9 +461,9 @@ void FSessionsLAN::StopLANSession()
 void FSessionsLAN::OnValidQueryPacketReceived(uint8* PacketData, int32 PacketLength, uint64 ClientNonce)
 {
 	// Iterate through all registered sessions and respond for each one that can be joinable
-	for (const TPair<FOnlineSessionIdHandle, TSharedRef<FSession>>& Entry : AllSessionsById)
+	for (const TPair<FOnlineSessionIdHandle, TSharedRef<FSessionCommon>>& Entry : AllSessionsById)
 	{
-		const TSharedRef<FSession>& Session = Entry.Value;
+		const TSharedRef<FSessionCommon>& Session = Entry.Value;
 
 		if (Session->SessionSettings.JoinPolicy == ESessionJoinPolicy::Public)
 		{
diff --git a/Engine/Plugins/Online/Experimental/OnlineServices/Source/OnlineServicesCommonEngineUtils/Public/Online/SessionsLAN.h b/Engine/Plugins/Online/Experimental/OnlineServices/Source/OnlineServicesCommonEngineUtils/Public/Online/SessionsLAN.h
--- a/Engine/Plugins/Online/Experimental/OnlineServices/Source/OnlineServicesCommonEngineUtils/Public/Online/SessionsLAN.h
+++ b/Engine/Plugins/Online/Experimental/OnlineServices/Source/OnlineServicesCommonEngineUtils/Public/Online/SessionsLAN.h
@@ -21,14 +21,14 @@ protected:
 	FOnlineSessionIdRegistryLAN(EOnlineServices ServicesType);
 };
 
-class ONLINESERVICESCOMMONENGINEUTILS_API FSessionLAN : public FSession
+class ONLINESERVICESCOMMONENGINEUTILS_API FSessionLAN : public FSessionCommon
 {
 public:
 	FSessionLAN();
 	FSessionLAN(const FSessionLAN& InSession) = default;
 
-	static FSessionLAN& Cast(FSession& InSession);
-	static const FSessionLAN& Cast(const FSession& InSession);
+	static FSessionLAN& Cast(FSessionCommon& InSession);
+	static const FSessionLAN& Cast(const ISession& InSession);
 
 private:
 	void Initialize();
diff --git a/Engine/Plugins/Online/Experimental/OnlineServices/Source/OnlineServicesInterface/Private/Online/Sessions.cpp b/Engine/Plugins/Online/Experimental/OnlineServices/Source/OnlineServicesInterface/Private/Online/Sessions.cpp
--- a/Engine/Plugins/Online/Experimental/OnlineServices/Source/OnlineServicesInterface/Private/Online/Sessions.cpp
+++ b/Engine/Plugins/Online/Experimental/OnlineServices/Source/OnlineServicesInterface/Private/Online/Sessions.cpp
@@ -187,17 +187,9 @@ FSessionSettingsUpdate& FSessionSettingsUpdate::operator+=(FSessionSettingsUpdat
 
 #undef MOVE_TOPTIONAL_IF_SET
 
-FSession::FSession()
+const FString ToLogString(const ISession& Session)
 {
-
-}
-
-FSession::FSession(const FSession& InSession)
-	: OwnerAccountId(InSession.OwnerAccountId)
-	, SessionId(InSession.SessionId)
-	, SessionSettings(InSession.SessionSettings)
-{
-
+	return Session.ToLogString();
 }
 
 const TCHAR* LexToString(EUISessionJoinRequestedSource UISessionJoinRequestedSource)
diff --git a/Engine/Plugins/Online/Experimental/OnlineServices/Source/OnlineServicesInterface/Public/Online/Sessions.h b/Engine/Plugins/Online/Experimental/OnlineServices/Source/OnlineServicesInterface/Public/Online/Sessions.h
--- a/Engine/Plugins/Online/Experimental/OnlineServices/Source/OnlineServicesInterface/Public/Online/Sessions.h
+++ b/Engine/Plugins/Online/Experimental/OnlineServices/Source/OnlineServicesInterface/Public/Online/Sessions.h
@@ -178,20 +178,15 @@ struct ONLINESERVICESINTERFACE_API FSessionSettings
 	FSessionSettings& operator+=(const FSessionSettingsUpdate& UpdatedValue);
 };
 
-struct ONLINESERVICESINTERFACE_API FSession
+class ISession
 {
-	FSession();
-	FSession(const FSession& InSession);
-
-	/** The user who currently owns the session */
-	FAccountId OwnerAccountId;
-
-	/** The id for the session, platform dependent */
-	FOnlineSessionIdHandle SessionId;
-
-	/** Set of session properties that can be altered by the session owner */
-	FSessionSettings SessionSettings;
+public:
+	virtual const FAccountId GetOwnerAccountId() const = 0;
+	virtual const FOnlineSessionIdHandle GetSessionId() const = 0;
+	virtual const FSessionSettings GetSessionSettings() const = 0;
+	virtual FString ToLogString() const = 0;
 };
+ONLINESERVICESINTERFACE_API const FString ToLogString(const ISession& Session);
 
 struct FSessionInvite
 {
@@ -221,7 +216,7 @@ struct FGetAllSessions
 
 	struct Result
 	{
-		TArray<TSharedRef<const FSession>> Sessions;
+		TArray<TSharedRef<const ISession>> Sessions;
 	};
 };
 
@@ -236,7 +231,7 @@ struct FGetSessionByName
 
 	struct Result
 	{
-		TSharedRef<const FSession> Session;
+		TSharedRef<const ISession> Session;
 
 		Result() = delete; // cannot default construct due to TSharedRef
 	};
@@ -255,7 +250,7 @@ struct FGetSessionById
 
 	struct Result
 	{
-		TSharedRef<const FSession> Session;
+		TSharedRef<const ISession> Session;
 
 		Result() = delete; // cannot default construct due to TSharedRef
 	};
@@ -886,12 +881,6 @@ BEGIN_ONLINE_STRUCT_META(FSessionSettingsUpdate)
 	ONLINE_STRUCT_FIELD(FSessionSettingsUpdate, RemovedRegisteredPlayers)
 END_ONLINE_STRUCT_META()
 
-BEGIN_ONLINE_STRUCT_META(FSession)
-	ONLINE_STRUCT_FIELD(FSession, OwnerAccountId),
-	ONLINE_STRUCT_FIELD(FSession, SessionId),
-	ONLINE_STRUCT_FIELD(FSession, SessionSettings)
-END_ONLINE_STRUCT_META()
-
 BEGIN_ONLINE_STRUCT_META(FSessionInvite)
 	ONLINE_STRUCT_FIELD(FSessionInvite, RecipientId),
 	ONLINE_STRUCT_FIELD(FSessionInvite, SenderId),
diff --git a/Engine/Plugins/Online/Experimental/OnlineServicesEOSGS/Source/Private/Online/OnlineServicesEOSGS.cpp b/Engine/Plugins/Online/Experimental/OnlineServicesEOSGS/Source/Private/Online/OnlineServicesEOSGS.cpp
--- a/Engine/Plugins/Online/Experimental/OnlineServicesEOSGS/Source/Private/Online/OnlineServicesEOSGS.cpp
+++ b/Engine/Plugins/Online/Experimental/OnlineServicesEOSGS/Source/Private/Online/OnlineServicesEOSGS.cpp
@@ -134,11 +134,11 @@ TOnlineResult<FGetResolvedConnectString> FOnlineServicesEOSGS::GetResolvedConnec
 		if (Result.IsOk())
 		{
 #if WITH_ENGINE
-			TSharedRef<const FSession> Session = Result.GetOkValue().Session;
+			TSharedRef<const ISession> Session = Result.GetOkValue().Session;
 
 			//It should look like this: "EOS:0002aeeb5b2d4388a3752dd6d31222ec:GameNetDriver:97"
 			FString NetDriverName = GetDefault<UNetDriverEOSBase>()->NetDriverName.ToString();
-			FInternetAddrEOS TempAddr(GetProductUserIdChecked(Session->OwnerAccountId), NetDriverName, GetTypeHash(NetDriverName));
+			FInternetAddrEOS TempAddr(GetProductUserIdChecked(Session->GetOwnerAccountId()), NetDriverName, GetTypeHash(NetDriverName));
 			return TOnlineResult<FGetResolvedConnectString>({ TempAddr.ToString(true) });
 #else
 			return TOnlineResult<FGetResolvedConnectString>(Errors::NotImplemented());
diff --git a/Engine/Plugins/Online/Experimental/OnlineServicesEOSGS/Source/Private/Online/SessionsEOSGS.cpp b/Engine/Plugins/Online/Experimental/OnlineServicesEOSGS/Source/Private/Online/SessionsEOSGS.cpp
--- a/Engine/Plugins/Online/Experimental/OnlineServicesEOSGS/Source/Private/Online/SessionsEOSGS.cpp
+++ b/Engine/Plugins/Online/Experimental/OnlineServicesEOSGS/Source/Private/Online/SessionsEOSGS.cpp
@@ -197,9 +197,9 @@ FSessionEOSGS::FSessionEOSGS(const EOS_HSessionDetails& InSessionDetailsHandle)
 	}
 }
 
-const FSessionEOSGS& FSessionEOSGS::Cast(const FSession& InSession)
+const FSessionEOSGS& FSessionEOSGS::Cast(const ISession& InSession)
 {
-	check(InSession.SessionId.GetOnlineServicesType() == EOnlineServices::Epic);
+	check(InSession.GetSessionId().GetOnlineServicesType() == EOnlineServices::Epic);
 
 	return static_cast<const FSessionEOSGS&>(InSession);
 }
@@ -291,7 +291,7 @@ void FSessionsEOSGS::HandleSessionInviteReceived(const EOS_Sessions_SessionInvit
 
 					FOnlineSessionInviteIdHandle InviteIdHandle = CreateSessionInviteId(InviteId);
 
-					TSharedRef<FSession> Session = Result.GetOkValue().Session;
+					TSharedRef<FSessionCommon> Session = Result.GetOkValue().Session;
 
 					TSharedRef<FSessionInvite> SessionInviteRef = MakeShared<FSessionInvite>(FSessionInvite{
 						ReceiverId,
@@ -739,8 +739,8 @@ void FSessionsEOSGS::WriteUpdateSessionModificationHandle(EOS_HSessionModificati
 	TOnlineResult<FGetSessionByName> GetSessionByNameResult = GetSessionByName({ SessionName });
 	check(GetSessionByNameResult.IsOk());
 
-	TSharedRef<const FSession> FoundSession = GetSessionByNameResult.GetOkValue().Session;
-	const FSessionSettings& SessionSettings = FoundSession->SessionSettings;
+	TSharedRef<const ISession> FoundSession = GetSessionByNameResult.GetOkValue().Session;
+	const FSessionSettings& SessionSettings = FoundSession->GetSessionSettings();
 
 	if (NewSettings.JoinPolicy.IsSet())
 	{
@@ -834,7 +834,7 @@ TOnlineAsyncOpHandle<FUpdateSession> FSessionsEOSGS::UpdateSession(FUpdateSessio
 {
 	// LAN Sessions
 	TOnlineResult<FGetSessionByName> Result = GetSessionByName({ Params.SessionName });
-	if (Result.IsOk() && Result.GetOkValue().Session->SessionSettings.bIsLANSession)
+	if (Result.IsOk() && Result.GetOkValue().Session->GetSessionSettings().bIsLANSession)
 	{
 		return FSessionsLAN::UpdateSession(MoveTemp(Params));
 	}
@@ -897,7 +897,7 @@ TOnlineAsyncOpHandle<FUpdateSession> FSessionsEOSGS::UpdateSession(FUpdateSessio
 			TOnlineResult<FGetMutableSessionByName> GetMutableSessionByNameResult = GetMutableSessionByName({ OpParams.SessionName });
 			if (GetMutableSessionByNameResult.IsOk())
 			{
-				TSharedRef<FSession> FoundSession = GetMutableSessionByNameResult.GetOkValue().Session;
+				TSharedRef<FSessionCommon> FoundSession = GetMutableSessionByNameResult.GetOkValue().Session;
 
 				// Now that the API Session update has processed successfully, we'll update our local session with the same data
 				FoundSession->SessionSettings += OpParams.Mutations;
@@ -1104,7 +1104,7 @@ TOnlineAsyncOpHandle<FLeaveSession> FSessionsEOSGS::LeaveSession(FLeaveSession::
 {
 	// LAN Sessions
 	TOnlineResult<FGetSessionByName> Result = GetSessionByName({ Params.SessionName });
-	if (Result.IsOk() && Result.GetOkValue().Session->SessionSettings.bIsLANSession)
+	if (Result.IsOk() && Result.GetOkValue().Session->GetSessionSettings().bIsLANSession)
 	{
 		return FSessionsLAN::LeaveSession(MoveTemp(Params));
 	}
@@ -1147,17 +1147,17 @@ TOnlineAsyncOpHandle<FLeaveSession> FSessionsEOSGS::LeaveSession(FLeaveSession::
 		TOnlineResult<FGetSessionByName> GetSessionByNameResult = GetSessionByName({ OpParams.SessionName });
 		if (GetSessionByNameResult.IsOk())
 		{
-			TSharedRef<const FSession> FoundSession = GetSessionByNameResult.GetOkValue().Session;
+			TSharedRef<const ISession> FoundSession = GetSessionByNameResult.GetOkValue().Session;
 
 			NamedSessionUserMap.FindChecked(OpParams.LocalAccountId).Remove(OpParams.SessionName);
 
-			if (FoundSession->SessionSettings.bPresenceEnabled)
+			if (FoundSession->GetSessionSettings().bPresenceEnabled)
 			{
 				PresenceSessionsUserMap.Remove(OpParams.LocalAccountId);
 			}
 
 			ClearSessionByName(OpParams.SessionName);
-			ClearSessionById(FoundSession->SessionId);
+			ClearSessionById(FoundSession->GetSessionId());
 		}
 
 		Op.SetResult(FLeaveSession::Result{ });
@@ -1376,7 +1376,7 @@ TOnlineAsyncOpHandle<FFindSessions> FSessionsEOSGS::FindSessions(FFindSessions::
 					{
 						if (TOnlineAsyncOpPtr<FFindSessions> StrongOp = WeakOp.Pin())
 						{
-							const TSharedRef<FSession>& Session = Result.GetOkValue().Session;
+							const TSharedRef<FSessionCommon>& Session = Result.GetOkValue().Session;
 
 							TArray<FOnlineSessionIdHandle>& SearchResults = SearchResultsUserMap.FindOrAdd(Result.GetOkValue().LocalAccountId);
 							SearchResults.Add(Session->SessionId);
@@ -1454,10 +1454,10 @@ TOnlineAsyncOpHandle<FJoinSession> FSessionsEOSGS::JoinSession(FJoinSession::Par
 		return Op->GetHandle();
 	}
 
-	const TSharedRef<const FSession>& FoundSession = GetSessionByIdResult.GetOkValue().Session;
+	const TSharedRef<const ISession>& FoundSession = GetSessionByIdResult.GetOkValue().Session;
 
 	// LAN Sessions
-	if (FoundSession->SessionSettings.bIsLANSession)
+	if (FoundSession->GetSessionSettings().bIsLANSession)
 	{
 		return FSessionsLAN::JoinSession(MoveTemp(Params));
 	}
@@ -1508,14 +1508,14 @@ TOnlineAsyncOpHandle<FJoinSession> FSessionsEOSGS::JoinSession(FJoinSession::Par
 			return;
 		}
 
-		const TSharedRef<const FSession>& FoundSession = GetSessionByIdResult.GetOkValue().Session;
+		const TSharedRef<const ISession>& FoundSession = GetSessionByIdResult.GetOkValue().Session;
 
 		// We start setup for the API call
 		EOS_Sessions_JoinSessionOptions JoinSessionOptions = { };
 		JoinSessionOptions.ApiVersion = EOS_SESSIONS_JOINSESSION_API_LATEST;
 		static_assert(EOS_SESSIONS_JOINSESSION_API_LATEST == 2, "EOS_Sessions_JoinSessionOptions updated, check new fields");
 
-		JoinSessionOptions.bPresenceEnabled = FoundSession->SessionSettings.bPresenceEnabled;
+		JoinSessionOptions.bPresenceEnabled = FoundSession->GetSessionSettings().bPresenceEnabled;
 
 		JoinSessionOptions.LocalUserId = GetProductUserIdChecked(OpParams.LocalAccountId);
 
@@ -1550,38 +1550,38 @@ TOnlineAsyncOpHandle<FJoinSession> FSessionsEOSGS::JoinSession(FJoinSession::Par
 			return;
 		}
 
-		const TSharedRef<const FSession>& FoundSession = GetSessionByIdResult.GetOkValue().Session;
+		const TSharedRef<const ISession>& FoundSession = GetSessionByIdResult.GetOkValue().Session;
 
 		LocalSessionsByName.Emplace(OpParams.SessionName, OpParams.SessionId);
 
 		NamedSessionUserMap.FindOrAdd(OpParams.LocalAccountId).AddUnique(OpParams.SessionName);
 
-		if (FoundSession->SessionSettings.bPresenceEnabled)
+		if (FoundSession->GetSessionSettings().bPresenceEnabled)
 		{
 			FOnlineSessionIdHandle& PresenceSessionId = PresenceSessionsUserMap.FindOrAdd(OpParams.LocalAccountId);
-			PresenceSessionId = FoundSession->SessionId;
+			PresenceSessionId = FoundSession->GetSessionId();
 		}
 
-			// After successfully joining a session, we'll remove all related invites if any are found
-			if (TMap<FOnlineSessionInviteIdHandle, TSharedRef<FSessionInvite>>* UserMap = SessionInvitesUserMap.Find(OpParams.LocalAccountId))
+		// After successfully joining a session, we'll remove all related invites if any are found
+		if (TMap<FOnlineSessionInviteIdHandle, TSharedRef<FSessionInvite>>* UserMap = SessionInvitesUserMap.Find(OpParams.LocalAccountId))
+		{
+			TArray<FOnlineSessionInviteIdHandle> InviteIdsToRemove;
+			for (const TPair<FOnlineSessionInviteIdHandle, TSharedRef<FSessionInvite>>& Entry : *UserMap)
 			{
-				TArray<FOnlineSessionInviteIdHandle> InviteIdsToRemove;
-				for (const TPair<FOnlineSessionInviteIdHandle, TSharedRef<FSessionInvite>>& Entry : *UserMap)
+				if (Entry.Value->SessionId == FoundSession->GetSessionId())
 				{
-				if (Entry.Value->SessionId == FoundSession->SessionId)
-					{
-						InviteIdsToRemove.Add(Entry.Key);
-					}
-				}
-				for (const FOnlineSessionInviteIdHandle& InviteId : InviteIdsToRemove)
-				{
-					UserMap->Remove(InviteId);
+					InviteIdsToRemove.Add(Entry.Key);
 				}
 			}
+			for (const FOnlineSessionInviteIdHandle& InviteId : InviteIdsToRemove)
+			{
+				UserMap->Remove(InviteId);
+			}
+		}
 
 		Op.SetResult(FJoinSession::Result{ });
 
-		FSessionJoined Event = { { OpParams.LocalAccountId }, FoundSession->SessionId };
+		FSessionJoined Event = { { OpParams.LocalAccountId }, FoundSession->GetSessionId() };
 
 			SessionEvents.OnSessionJoined.Broadcast(Event);
 
@@ -1848,7 +1848,7 @@ TOnlineAsyncOpHandle<FAddSessionMembers> FSessionsEOSGS::AddSessionMembers(FAddS
 {
 	// LAN Sessions
 	TOnlineResult<FGetSessionByName> Result = GetSessionByName({ Params.SessionName });
-	if (Result.IsOk() && Result.GetOkValue().Session->SessionSettings.bIsLANSession)
+	if (Result.IsOk() && Result.GetOkValue().Session->GetSessionSettings().bIsLANSession)
 	{
 		return FSessionsLAN::AddSessionMembers(MoveTemp(Params));
 	}
@@ -1918,7 +1918,7 @@ TOnlineAsyncOpHandle<FRemoveSessionMembers> FSessionsEOSGS::RemoveSessionMembers
 {
 	// LAN Sessions
 	TOnlineResult<FGetSessionByName> Result = GetSessionByName({ Params.SessionName });
-	if (Result.IsOk() && Result.GetOkValue().Session->SessionSettings.bIsLANSession)
+	if (Result.IsOk() && Result.GetOkValue().Session->GetSessionSettings().bIsLANSession)
 	{
 		return FSessionsLAN::RemoveSessionMembers(MoveTemp(Params));
 	}
diff --git a/Engine/Plugins/Online/Experimental/OnlineServicesEOSGS/Source/Public/Online/SessionsEOSGS.h b/Engine/Plugins/Online/Experimental/OnlineServicesEOSGS/Source/Public/Online/SessionsEOSGS.h
--- a/Engine/Plugins/Online/Experimental/OnlineServicesEOSGS/Source/Public/Online/SessionsEOSGS.h
+++ b/Engine/Plugins/Online/Experimental/OnlineServicesEOSGS/Source/Public/Online/SessionsEOSGS.h
@@ -87,7 +87,7 @@ public:
 	/** This constructor should only be used by BuildSessionFromDetailsHandle, after all user ids in the session have been resolved. */
 	FSessionEOSGS(const EOS_HSessionDetails& SessionDetailsHandle);
 
-	static const FSessionEOSGS& Cast(const FSession& InSession);
+	static const FSessionEOSGS& Cast(const ISession& InSession);
 
 public:
 	/** Session details handle */
@@ -171,7 +171,7 @@ struct FBuildSessionFromDetailsHandle
 		FAccountId LocalAccountId;
 
 		/** Session built from the details handle */
-		TSharedRef<FSession> Session;
+		TSharedRef<FSessionCommon> Session;
 	};
 };
 
diff --git a/Engine/Plugins/Online/Experimental/OnlineServicesNull/Source/Private/Online/OnlineServicesNull.cpp b/Engine/Plugins/Online/Experimental/OnlineServicesNull/Source/Private/Online/OnlineServicesNull.cpp
--- a/Engine/Plugins/Online/Experimental/OnlineServicesNull/Source/Private/Online/OnlineServicesNull.cpp
+++ b/Engine/Plugins/Online/Experimental/OnlineServicesNull/Source/Private/Online/OnlineServicesNull.cpp
@@ -94,14 +94,14 @@ TOnlineResult<FGetResolvedConnectString> FOnlineServicesNull::GetResolvedConnect
 		TOnlineResult<FGetAllSessions> JoinedSessions = SessionsPtr->GetAllSessions(MoveTemp(GetAllSessionsParams));
 		if (JoinedSessions.IsOk())
 		{
-			for (TSharedRef<const FSession>& Session : JoinedSessions.GetOkValue().Sessions)
+			for (TSharedRef<const ISession>& Session : JoinedSessions.GetOkValue().Sessions)
 			{
-				if (Session->SessionId == Params.SessionId)
+				if (Session->GetSessionId() == Params.SessionId)
 				{
-					if (Session->SessionSettings.CustomSettings.Contains(CONNECT_STRING_TAG))
+					if (Session->GetSessionSettings().CustomSettings.Contains(CONNECT_STRING_TAG))
 					{
 						FGetResolvedConnectString::Result Result;
-						Result.ResolvedConnectString = Session->SessionSettings.CustomSettings[CONNECT_STRING_TAG].Data.GetString();
+						Result.ResolvedConnectString = Session->GetSessionSettings().CustomSettings[CONNECT_STRING_TAG].Data.GetString();
 						return TOnlineResult<FGetResolvedConnectString>(Result);
 					}
 					else
diff --git a/Engine/Plugins/Online/Experimental/OnlineServicesOSSAdapter/Source/Private/Online/SessionsOSSAdapter.cpp b/Engine/Plugins/Online/Experimental/OnlineServicesOSSAdapter/Source/Private/Online/SessionsOSSAdapter.cpp
--- a/Engine/Plugins/Online/Experimental/OnlineServicesOSSAdapter/Source/Private/Online/SessionsOSSAdapter.cpp
+++ b/Engine/Plugins/Online/Experimental/OnlineServicesOSSAdapter/Source/Private/Online/SessionsOSSAdapter.cpp
@@ -243,7 +243,7 @@ FSessionOSSAdapter::FSessionOSSAdapter(const FOnlineSession& InSession)
 
 }
 
-const FSessionOSSAdapter& FSessionOSSAdapter::Cast(const FSession& InSession)
+const FSessionOSSAdapter& FSessionOSSAdapter::Cast(const ISession& InSession)
 {
 	// TODO: Check for ::Platform type if that's the Adapter type
 
@@ -278,7 +278,7 @@ void FSessionsOSSAdapter::Initialize()
 	{
 		FOnlineServicesOSSAdapter& ServicesOSSAdapter = static_cast<FOnlineServicesOSSAdapter&>(Services);
 
-		TSharedRef<FSession> Session = BuildV2Session(&InviteResult.Session);
+		TSharedRef<FSessionCommon> Session = BuildV2Session(&InviteResult.Session);
 
 		TSharedRef<FSessionInvite> SessionInvite = MakeShared<FSessionInvite>();
 		SessionInvite->SenderId = ServicesOSSAdapter.GetAccountIdRegistry().FindOrAddHandle(FromId.AsShared());
@@ -335,7 +335,7 @@ void FSessionsOSSAdapter::Initialize()
 		const TOnlineResult<FGetSessionByName> GetSessionByNameResult = GetSessionByName({ SessionName });
 		if (GetSessionByNameResult.IsOk())
 		{
-			const TSharedRef<const FSession>& FoundSession = GetSessionByNameResult.GetOkValue().Session;
+			const TSharedRef<const ISession>& FoundSession = GetSessionByNameResult.GetOkValue().Session;
 
 			const FOnlineServicesOSSAdapter& ServicesOSSAdapter = static_cast<FOnlineServicesOSSAdapter&>(Services);
 			const FUniqueNetIdRef TargetUniqueNetIdRef = TargetUniqueNetId.AsShared();
@@ -344,7 +344,7 @@ void FSessionsOSSAdapter::Initialize()
 			FSessionSettingsUpdate SettingsUpdate;
 			if (bJoined)
 			{
-				if (const FSessionMember* SessionMember = FoundSession->SessionSettings.SessionMembers.Find(TargetAccountId))
+				if (const FSessionMember* SessionMember = FoundSession->GetSessionSettings().SessionMembers.Find(TargetAccountId))
 				{
 					FSessionMemberUpdate SessionMemberUpdate;
 					SessionMemberUpdate.UpdatedMemberSettings.Append(SessionMember->MemberSettings);
@@ -369,14 +369,14 @@ void FSessionsOSSAdapter::Initialize()
 		const TOnlineResult<FGetSessionByName> GetSessionByNameResult = GetSessionByName({ SessionName });
 		if (GetSessionByNameResult.IsOk())
 		{
-			const TSharedRef<const FSession>& FoundSession = GetSessionByNameResult.GetOkValue().Session;
+			const TSharedRef<const ISession>& FoundSession = GetSessionByNameResult.GetOkValue().Session;
 
 			const FOnlineServicesOSSAdapter& ServicesOSSAdapter = static_cast<FOnlineServicesOSSAdapter&>(Services);
 			const FUniqueNetIdRef TargetUniqueNetIdRef = TargetUniqueNetId.AsShared();
 			const FAccountId TargetAccountId = ServicesOSSAdapter.GetAccountIdRegistry().FindOrAddHandle(TargetUniqueNetIdRef);
 
 			FSessionSettingsUpdate SettingsUpdate;
-			if (const FSessionMember* SessionMember = FoundSession->SessionSettings.SessionMembers.Find(TargetAccountId))
+			if (const FSessionMember* SessionMember = FoundSession->GetSessionSettings().SessionMembers.Find(TargetAccountId))
 			{
 				const ::FSessionSettings* MemberSettings = UpdatedSettings.MemberSettings.Find(TargetUniqueNetIdRef);
 				const FCustomSessionSettingsMap UpdatedMemberSettings = GetV2SessionSettings(*MemberSettings);
@@ -455,7 +455,7 @@ TOnlineAsyncOpHandle<FCreateSession> FSessionsOSSAdapter::CreateSession(FCreateS
 
 				const FNamedOnlineSession* V1Session = SessionsInterface->GetNamedSession(SessionName);
 
-				TSharedRef<FSession> V2Session = BuildV2Session(V1Session);
+				TSharedRef<FSessionCommon> V2Session = BuildV2Session(V1Session);
 
 				AllSessionsById.Emplace(V2Session->SessionId, V2Session);
 
@@ -508,7 +508,7 @@ TOnlineAsyncOpHandle<FUpdateSession> FSessionsOSSAdapter::UpdateSession(FUpdateS
 				const TOnlineResult<FGetMutableSessionByName> GetMutableSessionByNameResult = GetMutableSessionByName({ SessionName });
 				if (GetMutableSessionByNameResult.IsOk())
 				{
-					const TSharedRef<FSession>& FoundSession = GetMutableSessionByNameResult.GetOkValue().Session;
+					const TSharedRef<FSessionCommon>& FoundSession = GetMutableSessionByNameResult.GetOkValue().Session;
 
 					FSessionSettings& UpdatedV2Settings = FoundSession->SessionSettings;
 					UpdatedV2Settings += OpParams.Mutations;
@@ -526,10 +526,10 @@ TOnlineAsyncOpHandle<FUpdateSession> FSessionsOSSAdapter::UpdateSession(FUpdateS
 			}
 		});
 
-		const TOnlineResult<FGetSessionByName> GetSessionByNameResult = GetSessionByName({ OpParams.SessionName });
-		check(GetSessionByNameResult.IsOk());
+		const TOnlineResult<FGetMutableSessionByName> GetMutableSessionByNameResult = GetMutableSessionByName({ OpParams.SessionName });
+		check(GetMutableSessionByNameResult.IsOk());
 		
-		const TSharedRef<const FSession>& FoundSession = GetSessionByNameResult.GetOkValue().Session;
+		const TSharedRef<FSessionCommon>& FoundSession = GetMutableSessionByNameResult.GetOkValue().Session;
 
 		// We will update a copy here, and wait until the operation has completed successfully to update our local data
 		FSessionSettings UpdatedV2Settings = FoundSession->SessionSettings;
@@ -577,17 +577,17 @@ TOnlineAsyncOpHandle<FLeaveSession> FSessionsOSSAdapter::LeaveSession(FLeaveSess
 				TOnlineResult<FGetSessionByName> GetSessionByNameResult = GetSessionByName({ OpParams.SessionName });
 				if (GetSessionByNameResult.IsOk())
 				{
-					TSharedRef<const FSession> FoundSession = GetSessionByNameResult.GetOkValue().Session;
+					TSharedRef<const ISession> FoundSession = GetSessionByNameResult.GetOkValue().Session;
 
 					NamedSessionUserMap.FindChecked(OpParams.LocalAccountId).Remove(OpParams.SessionName);
 
-					if (FoundSession->SessionSettings.bPresenceEnabled)
+					if (FoundSession->GetSessionSettings().bPresenceEnabled)
 					{
 						PresenceSessionsUserMap.Remove(OpParams.LocalAccountId);
 					}
 
 					ClearSessionByName(OpParams.SessionName);
-					ClearSessionById(FoundSession->SessionId);
+					ClearSessionById(FoundSession->GetSessionId());
 				}
 
 				Op->SetResult({ });
@@ -662,10 +662,10 @@ TOnlineAsyncOpHandle<FFindSessions> FSessionsOSSAdapter::FindSessions(FFindSessi
 						{
 							TSharedRef<FOnlineSessionSearch>& PendingV1SessionSearch = PendingV1SessionSearchesPerUser.FindChecked(Op->GetParams().LocalAccountId);
 
-							TArray<TSharedRef<FSession>> FoundSessions = BuildV2SessionSearchResults(PendingV1SessionSearch->SearchResults);
+							TArray<TSharedRef<FSessionCommon>> FoundSessions = BuildV2SessionSearchResults(PendingV1SessionSearch->SearchResults);
 
 							TArray<FOnlineSessionIdHandle> SearchResults = SearchResultsUserMap.FindOrAdd(Op->GetParams().LocalAccountId);
-							for (const TSharedRef<FSession>& FoundSession : FoundSessions)
+							for (const TSharedRef<FSessionCommon>& FoundSession : FoundSessions)
 							{
 								AllSessionsById.Emplace(FoundSession->SessionId, FoundSession);
 								SearchResults.Add(FoundSession->SessionId);
@@ -725,9 +725,9 @@ TOnlineAsyncOpHandle<FFindSessions> FSessionsOSSAdapter::FindSessions(FFindSessi
 				{
 					if (bWasSuccessful)
 					{
-						TArray<TSharedRef<FSession>> FoundSessions = BuildV2SessionSearchResults(FriendSearchResults);
+						TArray<TSharedRef<FSessionCommon>> FoundSessions = BuildV2SessionSearchResults(FriendSearchResults);
 						TArray<FOnlineSessionIdHandle> SearchResults = SearchResultsUserMap.FindOrAdd(Op->GetParams().LocalAccountId);
-						for (const TSharedRef<FSession>& FoundSession : FoundSessions)
+						for (const TSharedRef<FSessionCommon>& FoundSession : FoundSessions)
 						{
 							AllSessionsById.Emplace(FoundSession->SessionId, FoundSession);
 							SearchResults.Add(FoundSession->SessionId);
@@ -764,9 +764,9 @@ TOnlineAsyncOpHandle<FFindSessions> FSessionsOSSAdapter::FindSessions(FFindSessi
 					{
 						TSharedRef<FOnlineSessionSearch>& PendingV1SessionSearch = PendingV1SessionSearchesPerUser.FindChecked(Op->GetParams().LocalAccountId);
 
-						TArray<TSharedRef<FSession>> FoundSessions = BuildV2SessionSearchResults(PendingV1SessionSearch->SearchResults);
+						TArray<TSharedRef<FSessionCommon>> FoundSessions = BuildV2SessionSearchResults(PendingV1SessionSearch->SearchResults);
 						TArray<FOnlineSessionIdHandle> SearchResults = SearchResultsUserMap.FindOrAdd(Op->GetParams().LocalAccountId);
-						for (const TSharedRef<FSession>& FoundSession : FoundSessions)
+						for (const TSharedRef<FSessionCommon>& FoundSession : FoundSessions)
 						{
 							AllSessionsById.Emplace(FoundSession->SessionId, FoundSession);
 							SearchResults.Add(FoundSession->SessionId);
@@ -876,7 +876,7 @@ TOnlineAsyncOpHandle<FStartMatchmaking> FSessionsOSSAdapter::StartMatchmaking(FS
 
 				const FStartMatchmaking::Params& OpParams = Op->GetParams();
 
-				TSharedRef<FSession> V2Session = BuildV2Session(SessionsInterface->GetNamedSession(SessionName));
+				TSharedRef<FSessionCommon> V2Session = BuildV2Session(SessionsInterface->GetNamedSession(SessionName));
 
 				AllSessionsById.Emplace(V2Session->SessionId, V2Session);
 
@@ -955,16 +955,16 @@ TOnlineAsyncOpHandle<FJoinSession> FSessionsOSSAdapter::JoinSession(FJoinSession
 					return;
 				}
 
-				TSharedRef<const FSession> FoundSession = GetSessionByIdResult.GetOkValue().Session;
+				TSharedRef<const ISession> FoundSession = GetSessionByIdResult.GetOkValue().Session;
 
 				LocalSessionsByName.Emplace(OpParams.SessionName, OpParams.SessionId);
 
 				NamedSessionUserMap.FindOrAdd(OpParams.LocalAccountId).AddUnique(OpParams.SessionName);
 
-				if (FoundSession->SessionSettings.bPresenceEnabled)
+				if (FoundSession->GetSessionSettings().bPresenceEnabled)
 				{
 					FOnlineSessionIdHandle& PresenceSessionId = PresenceSessionsUserMap.FindOrAdd(OpParams.LocalAccountId);
-					PresenceSessionId = FoundSession->SessionId;
+					PresenceSessionId = FoundSession->GetSessionId();
 				}
 
 				// After successfully joining a session, we'll remove all related invites if any are found
@@ -973,7 +973,7 @@ TOnlineAsyncOpHandle<FJoinSession> FSessionsOSSAdapter::JoinSession(FJoinSession
 					TArray<FOnlineSessionInviteIdHandle> InviteIdsToRemove;
 					for (const TPair<FOnlineSessionInviteIdHandle, TSharedRef<FSessionInvite>>& Entry : *UserMap)
 					{
-						if (Entry.Value->SessionId == FoundSession->SessionId)
+						if (Entry.Value->SessionId == FoundSession->GetSessionId())
 						{
 							InviteIdsToRemove.Add(Entry.Key);
 						}
@@ -986,7 +986,7 @@ TOnlineAsyncOpHandle<FJoinSession> FSessionsOSSAdapter::JoinSession(FJoinSession
 
 				Op->SetResult({ });
 
-				FSessionJoined Event{ { OpParams.LocalAccountId }, FoundSession->SessionId };
+				FSessionJoined Event{ { OpParams.LocalAccountId }, FoundSession->GetSessionId() };
 
 				SessionEvents.OnSessionJoined.Broadcast(Event);
 			}
@@ -999,11 +999,11 @@ TOnlineAsyncOpHandle<FJoinSession> FSessionsOSSAdapter::JoinSession(FJoinSession
 			return;
 		}
 
-		const TSharedRef<const FSession>& FoundSession = GetSessionByIdResult.GetOkValue().Session;
+		const TSharedRef<const ISession>& FoundSession = GetSessionByIdResult.GetOkValue().Session;
 
 		FOnlineSessionSearchResult SearchResult;
 
-		if (const FCustomSessionSetting* PingInMs = FoundSession->SessionSettings.CustomSettings.Find(OSS_ADAPTER_SESSIONS_PING_IN_MS))
+		if (const FCustomSessionSetting* PingInMs = FoundSession->GetSessionSettings().CustomSettings.Find(OSS_ADAPTER_SESSIONS_PING_IN_MS))
 		{
 			SearchResult.PingInMs = (int32)PingInMs->Data.GetInt64();
 		}
@@ -1318,12 +1318,12 @@ void FSessionsOSSAdapter::WriteV2SessionSettingsFromV1NamedSession(const FNamedO
 	}
 }
 
-FOnlineSession FSessionsOSSAdapter::BuildV1Session(const TSharedRef<const FSession> InSession) const
+FOnlineSession FSessionsOSSAdapter::BuildV1Session(const TSharedRef<const ISession> InSession) const
 {
 	return FOnlineSession(FSessionOSSAdapter::Cast(*InSession).GetV1Session());
 }
 
-TSharedRef<FSession> FSessionsOSSAdapter::BuildV2Session(const FOnlineSession* InSession) const
+TSharedRef<FSessionCommon> FSessionsOSSAdapter::BuildV2Session(const FOnlineSession* InSession) const
 {
 	TSharedRef<FSessionOSSAdapter> Session = MakeShared<FSessionOSSAdapter>(*InSession);
 
@@ -1341,13 +1341,13 @@ TSharedRef<FSession> FSessionsOSSAdapter::BuildV2Session(const FOnlineSession* I
 	return Session;
 }
 
-TArray<TSharedRef<FSession>> FSessionsOSSAdapter::BuildV2SessionSearchResults(const TArray<FOnlineSessionSearchResult>& SessionSearchResults) const
+TArray<TSharedRef<FSessionCommon>> FSessionsOSSAdapter::BuildV2SessionSearchResults(const TArray<FOnlineSessionSearchResult>& SessionSearchResults) const
 {
-	TArray<TSharedRef<FSession>> FoundSessions;
+	TArray<TSharedRef<FSessionCommon>> FoundSessions;
 
 	for (const FOnlineSessionSearchResult& SearchResult : SessionSearchResults)
 	{
-		TSharedRef<FSession> FoundSession = BuildV2Session(&SearchResult.Session);
+		TSharedRef<FSessionCommon> FoundSession = BuildV2Session(&SearchResult.Session);
 
 		FCustomSessionSetting PingInMs;
 		PingInMs.Data.Set((int64)SearchResult.PingInMs);
diff --git a/Engine/Plugins/Online/Experimental/OnlineServicesOSSAdapter/Source/Public/Online/SessionsOSSAdapter.h b/Engine/Plugins/Online/Experimental/OnlineServicesOSSAdapter/Source/Public/Online/SessionsOSSAdapter.h
--- a/Engine/Plugins/Online/Experimental/OnlineServicesOSSAdapter/Source/Public/Online/SessionsOSSAdapter.h
+++ b/Engine/Plugins/Online/Experimental/OnlineServicesOSSAdapter/Source/Public/Online/SessionsOSSAdapter.h
@@ -33,14 +33,14 @@ static FName OSS_ADAPTER_SESSIONS_PING_IN_MS = TEXT("OSS_ADAPTER_SESSIONS_PING_I
 static FName OSS_ADAPTER_SESSION_SEARCH_PING_BUCKET_SIZE = TEXT("OSS_ADAPTER_SESSION_SEARCH_PING_BUCKET_SIZE");
 static FName OSS_ADAPTER_SESSION_SEARCH_PLATFORM_HASH = TEXT("OSS_ADAPTER_SESSION_SEARCH_PLATFORM_HASH");
 
-class FSessionOSSAdapter : public FSession
+class FSessionOSSAdapter : public FSessionCommon
 {
 public:
 	FSessionOSSAdapter() = default;
 	FSessionOSSAdapter(const FSessionOSSAdapter& InSession) = default;
 	FSessionOSSAdapter(const FOnlineSession& InSession);
 
-	static const FSessionOSSAdapter& Cast(const FSession& InSession);
+	static const FSessionOSSAdapter& Cast(const ISession& InSession);
 
 	const FOnlineSession& GetV1Session() const;
 
@@ -85,11 +85,11 @@ private:
 	/** Writes V1 Session Settings information contained in Named session type into the passed V2 equivalent */
 	void WriteV2SessionSettingsFromV1NamedSession(const FNamedOnlineSession* InSession, FSessionSettings& OutSettings) const;
 	/** Builds a V1 Session data type from a V2 equivalent */
-	FOnlineSession BuildV1Session(const TSharedRef<const FSession> InSession) const;
+	FOnlineSession BuildV1Session(const TSharedRef<const ISession> InSession) const;
 	/** Builds a V2 Session data type from a V1 equivalent */
-	TSharedRef<FSession> BuildV2Session(const FOnlineSession* InSession) const;
+	TSharedRef<FSessionCommon> BuildV2Session(const FOnlineSession* InSession) const;
 	/** Builds a V2 Session Search Results array from the passed V1 equivalent types */
-	TArray<TSharedRef<FSession>> BuildV2SessionSearchResults(const TArray<FOnlineSessionSearchResult>& SessionSearchResults) const;
+	TArray<TSharedRef<FSessionCommon>> BuildV2SessionSearchResults(const TArray<FOnlineSessionSearchResult>& SessionSearchResults) const;
 
 	FOnlineSessionIdRegistryOSSAdapter& GetSessionIdRegistry() const;
 	FOnlineSessionInviteIdRegistryOSSAdapter& GetSessionInviteIdRegistry() const;
