Addition of methods to streamline adding and removing references to sessions

diff --git a/Engine/Plugins/Online/Experimental/OnlineServices/Source/OnlineServicesCommon/Private/Online/SessionsCommon.cpp b/Engine/Plugins/Online/Experimental/OnlineServices/Source/OnlineServicesCommon/Private/Online/SessionsCommon.cpp
--- a/Engine/Plugins/Online/Experimental/OnlineServices/Source/OnlineServicesCommon/Private/Online/SessionsCommon.cpp
+++ b/Engine/Plugins/Online/Experimental/OnlineServices/Source/OnlineServicesCommon/Private/Online/SessionsCommon.cpp
@@ -360,6 +360,74 @@ namespace UE::Online {
 		}
 	}
 
+	void FSessionsCommon::AddSessionInvite(const TSharedRef<FSessionInvite> SessionInvite, const TSharedRef<FSessionCommon> Session, const FAccountId& LocalAccountId)
+	{
+		AllSessionsById.Emplace(Session->SessionId, Session);
+
+		TMap<FOnlineSessionInviteIdHandle, TSharedRef<FSessionInvite>>& SessionInvitesMap = SessionInvitesUserMap.FindOrAdd(LocalAccountId);
+		SessionInvitesMap.Emplace(SessionInvite->InviteId, SessionInvite);
+	}
+
+	void FSessionsCommon::AddSearchResult(const TSharedRef<FSessionCommon> Session, const FAccountId& LocalAccountId)
+	{
+		AllSessionsById.Emplace(Session->SessionId, Session);
+
+		TArray<FOnlineSessionIdHandle>& SearchResults = SearchResultsUserMap.FindOrAdd(LocalAccountId);
+		SearchResults.Add(Session->SessionId);
+	}
+
+	void FSessionsCommon::AddSessionWithReferences(const TSharedRef<FSessionCommon> Session, const FName& SessionName, const FAccountId& LocalAccountId, bool bIsPresenceSession)
+	{
+		AllSessionsById.Emplace(Session->SessionId, Session);
+
+		AddSessionReferences(Session->SessionId, SessionName, LocalAccountId, bIsPresenceSession);
+	}
+
+	void FSessionsCommon::AddSessionReferences(const FOnlineSessionIdHandle SessionId, const FName& SessionName, const FAccountId& LocalAccountId, bool bIsPresenceSession)
+	{
+		LocalSessionsByName.Emplace(SessionName, SessionId);
+
+		NamedSessionUserMap.FindOrAdd(LocalAccountId).AddUnique(SessionName);
+
+		if (bIsPresenceSession)
+		{
+			SetPresenceSession({ LocalAccountId, SessionId });
+		}
+	}
+
+	void FSessionsCommon::ClearSessionInvitesForSession(const FAccountId& LocalAccountId, const FOnlineSessionIdHandle SessionId)
+	{
+		if (TMap<FOnlineSessionInviteIdHandle, TSharedRef<FSessionInvite>>* UserMap = SessionInvitesUserMap.Find(LocalAccountId))
+		{
+			TArray<FOnlineSessionInviteIdHandle> InviteIdsToRemove;
+			for (const TPair<FOnlineSessionInviteIdHandle, TSharedRef<FSessionInvite>>& Entry : *UserMap)
+			{
+				if (Entry.Value->SessionId == SessionId)
+				{
+					InviteIdsToRemove.Add(Entry.Key);
+				}
+			}
+
+			for (const FOnlineSessionInviteIdHandle& InviteId : InviteIdsToRemove)
+			{
+				UserMap->Remove(InviteId);
+			}
+		}
+	}
+
+	void FSessionsCommon::ClearSessionReferences(const FOnlineSessionIdHandle SessionId, const FName& SessionName, const FAccountId& LocalAccountId, bool bIsPresenceSession)
+	{
+		NamedSessionUserMap.FindChecked(LocalAccountId).Remove(SessionName);
+
+		if (bIsPresenceSession)
+		{
+			ClearPresenceSession({ LocalAccountId });
+		}
+
+		ClearSessionByName(SessionName);
+		ClearSessionById(SessionId);
+	}
+
 	void FSessionsCommon::ClearSessionByName(const FName& SessionName)
 	{
 		for (const TPair<FAccountId, TArray<FName>>& Entry : NamedSessionUserMap)
diff --git a/Engine/Plugins/Online/Experimental/OnlineServices/Source/OnlineServicesCommon/Public/Online/SessionsCommon.h b/Engine/Plugins/Online/Experimental/OnlineServices/Source/OnlineServicesCommon/Public/Online/SessionsCommon.h
--- a/Engine/Plugins/Online/Experimental/OnlineServices/Source/OnlineServicesCommon/Public/Online/SessionsCommon.h
+++ b/Engine/Plugins/Online/Experimental/OnlineServices/Source/OnlineServicesCommon/Public/Online/SessionsCommon.h
@@ -228,8 +228,13 @@ protected:
 	TOnlineResult<FGetMutableSessionByName> GetMutableSessionByName(FGetMutableSessionByName::Params&& Params) const;
 	TOnlineResult<FGetMutableSessionById> GetMutableSessionById(FGetMutableSessionById::Params&& Params) const;
 
-	void ClearSessionByName(const FName& SessionName);
-	void ClearSessionById(const FOnlineSessionIdHandle& SessionId);
+	void AddSessionInvite(const TSharedRef<FSessionInvite> SessionInvite, const TSharedRef<FSessionCommon> Session, const FAccountId& LocalAccountId);
+	void AddSearchResult(const TSharedRef<FSessionCommon> Session, const FAccountId& LocalAccountId);
+	void AddSessionWithReferences(const TSharedRef<FSessionCommon> Session, const FName& SessionName, const FAccountId& LocalAccountId, bool bIsPresenceSession);
+	void AddSessionReferences(const FOnlineSessionIdHandle SessionId, const FName& SessionName, const FAccountId& LocalAccountId, bool bIsPresenceSession);
+
+	void ClearSessionInvitesForSession(const FAccountId& LocalAccountId, const FOnlineSessionIdHandle SessionId);
+	void ClearSessionReferences(const FOnlineSessionIdHandle SessionId, const FName& SessionName, const FAccountId& LocalAccountId, bool bIsPresenceSession);
 
 	TOnlineResult<FAddSessionMember> AddSessionMemberImpl(const FAddSessionMember::Params& Params);
 	TOnlineResult<FRemoveSessionMember> RemoveSessionMemberImpl(const FRemoveSessionMember::Params& Params);
@@ -268,6 +273,9 @@ protected:
 private:
 	TOptional<FOnlineError> CheckSessionExistsByName(const FAccountId& LocalAccountId, const FName& SessionName);
 
+	void ClearSessionByName(const FName& SessionName);
+	void ClearSessionById(const FOnlineSessionIdHandle& SessionId);
+
 protected:
 
 	struct FSessionEvents
diff --git a/Engine/Plugins/Online/Experimental/OnlineServices/Source/OnlineServicesCommonEngineUtils/Private/Online/SessionsLAN.cpp b/Engine/Plugins/Online/Experimental/OnlineServices/Source/OnlineServicesCommonEngineUtils/Private/Online/SessionsLAN.cpp
--- a/Engine/Plugins/Online/Experimental/OnlineServices/Source/OnlineServicesCommonEngineUtils/Private/Online/SessionsLAN.cpp
+++ b/Engine/Plugins/Online/Experimental/OnlineServices/Source/OnlineServicesCommonEngineUtils/Private/Online/SessionsLAN.cpp
@@ -160,16 +160,7 @@ TOnlineAsyncOpHandle<FCreateSession> FSessionsLAN::CreateSession(FCreateSession:
 		NewSessionLANRef->SessionSettings.SessionMembers.Emplace(OpParams.LocalAccountId, OpParams.SessionMemberData);
 
 		// We save the local object for the session, and set up the appropriate references
-		AllSessionsById.Emplace(NewSessionLANRef->SessionId, NewSessionLANRef);
-
-		LocalSessionsByName.Emplace(OpParams.SessionName, NewSessionLANRef->SessionId);
-
-		NamedSessionUserMap.FindOrAdd(OpParams.LocalAccountId).AddUnique(OpParams.SessionName);
-
-		if (NewSessionLANRef->SessionSettings.bPresenceEnabled)
-		{
-			SetPresenceSession({ OpParams.LocalAccountId, NewSessionLANRef->SessionId });
-		}
+		AddSessionWithReferences(NewSessionLANRef, OpParams.SessionName, OpParams.LocalAccountId, NewSessionLANRef->SessionSettings.bPresenceEnabled);
 
 		Op.SetResult({ });
 	})
@@ -305,14 +296,7 @@ TOnlineAsyncOpHandle<FJoinSession> FSessionsLAN::JoinSession(FJoinSession::Param
 		const TSharedRef<const ISession>& FoundSession = GetSessionByIdResult.GetOkValue().Session;
 
 		// We set up the appropriate references for the session
-		LocalSessionsByName.Emplace(OpParams.SessionName, FoundSession->GetSessionId());
-
-		NamedSessionUserMap.FindOrAdd(OpParams.LocalAccountId).AddUnique(OpParams.SessionName);
-
-		if (FoundSession->GetSessionSettings().bPresenceEnabled)
-		{
-			SetPresenceSession({ OpParams.LocalAccountId, FoundSession->GetSessionId() });
-		}
+		AddSessionReferences(FoundSession->GetSessionId(), OpParams.SessionName, OpParams.LocalAccountId, FoundSession->GetSessionSettings().bPresenceEnabled);
 
 		Op.SetResult(FJoinSession::Result{ });
 
@@ -350,15 +334,7 @@ TOnlineAsyncOpHandle<FLeaveSession> FSessionsLAN::LeaveSession(FLeaveSession::Pa
 				StopLANSession();
 			}
 
-			NamedSessionUserMap.FindChecked(OpParams.LocalAccountId).Remove(OpParams.SessionName);
-
-			if (FoundSession->GetSessionSettings().bPresenceEnabled)
-			{
-				ClearPresenceSession({ OpParams.LocalAccountId });
-			}
-
-			ClearSessionByName(OpParams.SessionName);
-			ClearSessionById(FoundSession->GetSessionId());
+			ClearSessionReferences(FoundSession->GetSessionId(), OpParams.SessionName, OpParams.LocalAccountId, FoundSession->GetSessionSettings().bPresenceEnabled);
 		}
 
 		Op.SetResult(FLeaveSession::Result{ });
@@ -488,17 +464,12 @@ void FSessionsLAN::OnValidQueryPacketReceived(uint8* PacketData, int32 PacketLen
 
 void FSessionsLAN::OnValidResponsePacketReceived(uint8* PacketData, int32 PacketLength, const FAccountId LocalAccountId)
 {
-	if (TArray<FOnlineSessionIdHandle>* SearchResults = SearchResultsUserMap.Find(LocalAccountId))
-	{
-		TSharedRef<FSessionLAN> Session = MakeShared<FSessionLAN>();
+	TSharedRef<FSessionLAN> Session = MakeShared<FSessionLAN>();
 
-		FNboSerializeFromBuffer Packet(PacketData, PacketLength);
-		ReadSessionFromPacket(Packet, *Session);
+	FNboSerializeFromBuffer Packet(PacketData, PacketLength);
+	ReadSessionFromPacket(Packet, *Session);
 
-		SearchResults->Add(Session->SessionId);
-
-		AllSessionsById.Emplace(Session->SessionId, Session);
-	}
+	AddSearchResult(Session, LocalAccountId);
 }
 
 void FSessionsLAN::OnLANSearchTimeout(const FAccountId LocalAccountId)
diff --git a/Engine/Plugins/Online/Experimental/OnlineServicesEOSGS/Source/Private/Online/SessionsEOSGS.cpp b/Engine/Plugins/Online/Experimental/OnlineServicesEOSGS/Source/Private/Online/SessionsEOSGS.cpp
--- a/Engine/Plugins/Online/Experimental/OnlineServicesEOSGS/Source/Private/Online/SessionsEOSGS.cpp
+++ b/Engine/Plugins/Online/Experimental/OnlineServicesEOSGS/Source/Private/Online/SessionsEOSGS.cpp
@@ -300,9 +300,7 @@ void FSessionsEOSGS::HandleSessionInviteReceived(const EOS_Sessions_SessionInvit
 						Session->SessionId
 						});
 
-					SessionInvitesUserMap.FindOrAdd(ReceiverId).Emplace(InviteIdHandle, SessionInviteRef);
-
-					AllSessionsById.Emplace(Session->SessionId, Session);
+					AddSessionInvite(SessionInviteRef, Session, ReceiverId);
 
 					FSessionInviteReceived Event = { ReceiverId, SessionInviteRef };
 
@@ -495,16 +493,7 @@ TOnlineAsyncOpHandle<FCreateSession> FSessionsEOSGS::CreateSession(FCreateSessio
 			NewSessionEOSGSRef->SessionSettings = OpParams.SessionSettings;
 			NewSessionEOSGSRef->SessionId = CreateSessionId(Result.GetOkValue().NewSessionId);
 
-			AllSessionsById.Emplace(NewSessionEOSGSRef->SessionId, NewSessionEOSGSRef);
-
-			LocalSessionsByName.Emplace(OpParams.SessionName, NewSessionEOSGSRef->SessionId);
-
-			NamedSessionUserMap.FindOrAdd(OpParams.LocalAccountId).AddUnique(OpParams.SessionName);
-
-			if (NewSessionEOSGSRef->SessionSettings.bPresenceEnabled)
-			{
-				SetPresenceSession({ OpParams.LocalAccountId, NewSessionEOSGSRef->SessionId });
-			}
+			AddSessionWithReferences(NewSessionEOSGSRef, OpParams.SessionName, OpParams.LocalAccountId, NewSessionEOSGSRef->SessionSettings.bPresenceEnabled);
 
 			Op.SetResult({ });
 		}
@@ -1148,15 +1137,7 @@ TOnlineAsyncOpHandle<FLeaveSession> FSessionsEOSGS::LeaveSession(FLeaveSession::
 		{
 			TSharedRef<const ISession> FoundSession = GetSessionByNameResult.GetOkValue().Session;
 
-			NamedSessionUserMap.FindChecked(OpParams.LocalAccountId).Remove(OpParams.SessionName);
-
-			if (FoundSession->GetSessionSettings().bPresenceEnabled)
-			{
-				ClearPresenceSession({ OpParams.LocalAccountId });
-			}
-
-			ClearSessionByName(OpParams.SessionName);
-			ClearSessionById(FoundSession->GetSessionId());
+			ClearSessionReferences(FoundSession->GetSessionId(), OpParams.SessionName, OpParams.LocalAccountId, FoundSession->GetSessionSettings().bPresenceEnabled);
 		}
 
 		Op.SetResult(FLeaveSession::Result{ });
@@ -1374,12 +1355,7 @@ TOnlineAsyncOpHandle<FFindSessions> FSessionsEOSGS::FindSessions(FFindSessions::
 					{
 						if (TOnlineAsyncOpPtr<FFindSessions> StrongOp = WeakOp.Pin())
 						{
-							const TSharedRef<FSessionCommon>& Session = Result.GetOkValue().Session;
-
-							TArray<FOnlineSessionIdHandle>& SearchResults = SearchResultsUserMap.FindOrAdd(Result.GetOkValue().LocalAccountId);
-							SearchResults.Add(Session->SessionId);
-
-							AllSessionsById.Emplace(Session->SessionId, Session);
+							AddSearchResult(Result.GetOkValue().Session, Result.GetOkValue().LocalAccountId);
 
 							BuildSessionPromise->EmplaceValue(Result.GetOkValue());
 						}
@@ -1550,31 +1526,10 @@ TOnlineAsyncOpHandle<FJoinSession> FSessionsEOSGS::JoinSession(FJoinSession::Par
 
 		const TSharedRef<const ISession>& FoundSession = GetSessionByIdResult.GetOkValue().Session;
 
-		LocalSessionsByName.Emplace(OpParams.SessionName, OpParams.SessionId);
-
-		NamedSessionUserMap.FindOrAdd(OpParams.LocalAccountId).AddUnique(OpParams.SessionName);
-
-		if (FoundSession->GetSessionSettings().bPresenceEnabled)
-		{
-			SetPresenceSession({ OpParams.LocalAccountId, FoundSession->GetSessionId() });
-		}
+		AddSessionReferences(FoundSession->GetSessionId(), OpParams.SessionName, OpParams.LocalAccountId, FoundSession->GetSessionSettings().bPresenceEnabled);
 
 		// After successfully joining a session, we'll remove all related invites if any are found
-		if (TMap<FOnlineSessionInviteIdHandle, TSharedRef<FSessionInvite>>* UserMap = SessionInvitesUserMap.Find(OpParams.LocalAccountId))
-		{
-			TArray<FOnlineSessionInviteIdHandle> InviteIdsToRemove;
-			for (const TPair<FOnlineSessionInviteIdHandle, TSharedRef<FSessionInvite>>& Entry : *UserMap)
-			{
-				if (Entry.Value->SessionId == FoundSession->GetSessionId())
-				{
-					InviteIdsToRemove.Add(Entry.Key);
-				}
-			}
-			for (const FOnlineSessionInviteIdHandle& InviteId : InviteIdsToRemove)
-			{
-				UserMap->Remove(InviteId);
-			}
-		}
+		ClearSessionInvitesForSession(OpParams.LocalAccountId, FoundSession->GetSessionId());
 
 		Op.SetResult(FJoinSession::Result{ });
 
diff --git a/Engine/Plugins/Online/Experimental/OnlineServicesOSSAdapter/Source/Private/Online/SessionsOSSAdapter.cpp b/Engine/Plugins/Online/Experimental/OnlineServicesOSSAdapter/Source/Private/Online/SessionsOSSAdapter.cpp
--- a/Engine/Plugins/Online/Experimental/OnlineServicesOSSAdapter/Source/Private/Online/SessionsOSSAdapter.cpp
+++ b/Engine/Plugins/Online/Experimental/OnlineServicesOSSAdapter/Source/Private/Online/SessionsOSSAdapter.cpp
@@ -286,9 +286,7 @@ void FSessionsOSSAdapter::Initialize()
 		SessionInvite->InviteId = GetSessionInviteIdRegistry().BasicRegistry.FindOrAddHandle(AppId);
 		SessionInvite->SessionId = Session->SessionId;
 
-		SessionInvitesUserMap.FindOrAdd(SessionInvite->RecipientId).Emplace(SessionInvite->InviteId, SessionInvite);
-
-		AllSessionsById.Emplace(Session->SessionId, Session);
+		AddSessionInvite(SessionInvite, Session, SessionInvite->RecipientId);
 
 		FSessionInviteReceived Event{ ServicesOSSAdapter.GetAccountIdRegistry().FindOrAddHandle(UserId.AsShared()), SessionInvite };
 
@@ -457,16 +455,7 @@ TOnlineAsyncOpHandle<FCreateSession> FSessionsOSSAdapter::CreateSession(FCreateS
 
 				TSharedRef<FSessionCommon> V2Session = BuildV2Session(V1Session);
 
-				AllSessionsById.Emplace(V2Session->SessionId, V2Session);
-
-				LocalSessionsByName.Emplace(OpParams.SessionName, V2Session->SessionId);
-
-				NamedSessionUserMap.FindOrAdd(OpParams.LocalAccountId).AddUnique(OpParams.SessionName);
-
-				if (V2Session->SessionSettings.bPresenceEnabled)
-				{
-					SetPresenceSession({ OpParams.LocalAccountId, V2Session->SessionId });
-				}
+				AddSessionWithReferences(V2Session, OpParams.SessionName, OpParams.LocalAccountId, V2Session->SessionSettings.bPresenceEnabled);
 
 				Op->SetResult({ });
 			}
@@ -584,15 +573,7 @@ TOnlineAsyncOpHandle<FLeaveSession> FSessionsOSSAdapter::LeaveSession(FLeaveSess
 				{
 					TSharedRef<const ISession> FoundSession = GetSessionByNameResult.GetOkValue().Session;
 
-					NamedSessionUserMap.FindChecked(OpParams.LocalAccountId).Remove(OpParams.SessionName);
-
-					if (FoundSession->GetSessionSettings().bPresenceEnabled)
-					{
-						ClearPresenceSession({ OpParams.LocalAccountId });
-					}
-
-					ClearSessionByName(OpParams.SessionName);
-					ClearSessionById(FoundSession->GetSessionId());
+					ClearSessionReferences(FoundSession->GetSessionId(), OpParams.SessionName, OpParams.LocalAccountId, FoundSession->GetSessionSettings().bPresenceEnabled);
 				}
 
 				Op->SetResult({ });
@@ -668,14 +649,12 @@ TOnlineAsyncOpHandle<FFindSessions> FSessionsOSSAdapter::FindSessions(FFindSessi
 							TSharedRef<FOnlineSessionSearch>& PendingV1SessionSearch = PendingV1SessionSearchesPerUser.FindChecked(Op->GetParams().LocalAccountId);
 
 							TArray<TSharedRef<FSessionCommon>> FoundSessions = BuildV2SessionSearchResults(PendingV1SessionSearch->SearchResults);
-
-							TArray<FOnlineSessionIdHandle> SearchResults = SearchResultsUserMap.FindOrAdd(Op->GetParams().LocalAccountId);
 							for (const TSharedRef<FSessionCommon>& FoundSession : FoundSessions)
 							{
-								AllSessionsById.Emplace(FoundSession->SessionId, FoundSession);
-								SearchResults.Add(FoundSession->SessionId);
+								AddSearchResult(FoundSession, Op->GetParams().LocalAccountId);
 							}
 
+							TArray<FOnlineSessionIdHandle> SearchResults = SearchResultsUserMap.FindOrAdd(Op->GetParams().LocalAccountId);
 							Op->SetResult({ SearchResults });
 						}
 						else
@@ -731,13 +710,12 @@ TOnlineAsyncOpHandle<FFindSessions> FSessionsOSSAdapter::FindSessions(FFindSessi
 					if (bWasSuccessful)
 					{
 						TArray<TSharedRef<FSessionCommon>> FoundSessions = BuildV2SessionSearchResults(FriendSearchResults);
-						TArray<FOnlineSessionIdHandle> SearchResults = SearchResultsUserMap.FindOrAdd(Op->GetParams().LocalAccountId);
 						for (const TSharedRef<FSessionCommon>& FoundSession : FoundSessions)
 						{
-							AllSessionsById.Emplace(FoundSession->SessionId, FoundSession);
-							SearchResults.Add(FoundSession->SessionId);
+							AddSearchResult(FoundSession, Op->GetParams().LocalAccountId);
 						}
 
+						TArray<FOnlineSessionIdHandle> SearchResults = SearchResultsUserMap.FindOrAdd(Op->GetParams().LocalAccountId);
 						Op->SetResult({ SearchResults });
 					}
 					else
@@ -770,13 +748,12 @@ TOnlineAsyncOpHandle<FFindSessions> FSessionsOSSAdapter::FindSessions(FFindSessi
 						TSharedRef<FOnlineSessionSearch>& PendingV1SessionSearch = PendingV1SessionSearchesPerUser.FindChecked(Op->GetParams().LocalAccountId);
 
 						TArray<TSharedRef<FSessionCommon>> FoundSessions = BuildV2SessionSearchResults(PendingV1SessionSearch->SearchResults);
-						TArray<FOnlineSessionIdHandle> SearchResults = SearchResultsUserMap.FindOrAdd(Op->GetParams().LocalAccountId);
 						for (const TSharedRef<FSessionCommon>& FoundSession : FoundSessions)
 						{
-							AllSessionsById.Emplace(FoundSession->SessionId, FoundSession);
-							SearchResults.Add(FoundSession->SessionId);
+							AddSearchResult(FoundSession, Op->GetParams().LocalAccountId);
 						}
 
+						TArray<FOnlineSessionIdHandle> SearchResults = SearchResultsUserMap.FindOrAdd(Op->GetParams().LocalAccountId);
 						Op->SetResult({ SearchResults });
 					}
 					else
@@ -881,16 +858,7 @@ TOnlineAsyncOpHandle<FStartMatchmaking> FSessionsOSSAdapter::StartMatchmaking(FS
 
 				TSharedRef<FSessionCommon> V2Session = BuildV2Session(SessionsInterface->GetNamedSession(SessionName));
 
-				AllSessionsById.Emplace(V2Session->SessionId, V2Session);
-
-				LocalSessionsByName.Emplace(OpParams.SessionName, V2Session->SessionId);
-
-				NamedSessionUserMap.FindOrAdd(OpParams.LocalAccountId).AddUnique(OpParams.SessionName);
-
-				if (V2Session->SessionSettings.bPresenceEnabled)
-				{
-					SetPresenceSession({ OpParams.LocalAccountId, V2Session->SessionId });
-				}
+				AddSessionWithReferences(V2Session, OpParams.SessionName, OpParams.LocalAccountId, V2Session->SessionSettings.bPresenceEnabled);
 
 				Op->SetResult({ });
 
@@ -965,31 +933,10 @@ TOnlineAsyncOpHandle<FJoinSession> FSessionsOSSAdapter::JoinSession(FJoinSession
 
 				TSharedRef<const ISession> FoundSession = GetSessionByIdResult.GetOkValue().Session;
 
-				LocalSessionsByName.Emplace(OpParams.SessionName, OpParams.SessionId);
-
-				NamedSessionUserMap.FindOrAdd(OpParams.LocalAccountId).AddUnique(OpParams.SessionName);
-
-				if (FoundSession->GetSessionSettings().bPresenceEnabled)
-				{
-					SetPresenceSession({ OpParams.LocalAccountId, FoundSession->GetSessionId() });
-				}
+				AddSessionReferences(FoundSession->GetSessionId(), OpParams.SessionName, OpParams.LocalAccountId, FoundSession->GetSessionSettings().bPresenceEnabled);
 
 				// After successfully joining a session, we'll remove all related invites if any are found
-				if (TMap<FOnlineSessionInviteIdHandle, TSharedRef<FSessionInvite>>* UserMap = SessionInvitesUserMap.Find(OpParams.LocalAccountId))
-				{
-					TArray<FOnlineSessionInviteIdHandle> InviteIdsToRemove;
-					for (const TPair<FOnlineSessionInviteIdHandle, TSharedRef<FSessionInvite>>& Entry : *UserMap)
-					{
-						if (Entry.Value->SessionId == FoundSession->GetSessionId())
-						{
-							InviteIdsToRemove.Add(Entry.Key);
-						}
-					}
-					for (const FOnlineSessionInviteIdHandle& InviteId : InviteIdsToRemove)
-					{
-						UserMap->Remove(InviteId);
-					}
-				}
+				ClearSessionInvitesForSession(OpParams.LocalAccountId, FoundSession->GetSessionId());
 
 				Op->SetResult({ });
 
