OSSv2: Remove the redundant Super:: after the renaming of LoadConfig() to UpdateConfig().

diff --git a/Engine/Plugins/Online/Experimental/OnlineServicesEOS/Source/Public/Online/ExternalUIEOS.h b/Engine/Plugins/Online/Experimental/OnlineServicesEOS/Source/Public/Online/ExternalUIEOS.h
--- a/Engine/Plugins/Online/Experimental/OnlineServicesEOS/Source/Public/Online/ExternalUIEOS.h
+++ b/Engine/Plugins/Online/Experimental/OnlineServicesEOS/Source/Public/Online/ExternalUIEOS.h
@@ -33,7 +33,6 @@ public:
 	virtual TOnlineAsyncOpHandle<FExternalUIShowFriendsUI> ShowFriendsUI(FExternalUIShowFriendsUI::Params&& Params) override;
 
 protected:
-	using Super::LoadConfig;
 	/** Handle to EOS UI */
 	EOS_HUI UIHandle;
 	/** Config */
diff --git a/Engine/Plugins/Online/Experimental/OnlineServicesEOSGS/Source/Private/Online/TitleFileEOSGS.cpp b/Engine/Plugins/Online/Experimental/OnlineServicesEOSGS/Source/Private/Online/TitleFileEOSGS.cpp
--- a/Engine/Plugins/Online/Experimental/OnlineServicesEOSGS/Source/Private/Online/TitleFileEOSGS.cpp
+++ b/Engine/Plugins/Online/Experimental/OnlineServicesEOSGS/Source/Private/Online/TitleFileEOSGS.cpp
@@ -30,7 +30,7 @@ void FTitleFileEOSGS::Initialize()
 void FTitleFileEOSGS::UpdateConfig()
 {
 	Super::UpdateConfig();
-	Super::LoadConfig(Config);
+	LoadConfig(Config);
 }
 
 TOnlineAsyncOpHandle<FTitleFileEnumerateFiles> FTitleFileEOSGS::EnumerateFiles(FTitleFileEnumerateFiles::Params&& InParams)
diff --git a/Engine/Plugins/Online/Experimental/OnlineServicesEOSGS/Source/Private/Online/UserFileEOSGS.cpp b/Engine/Plugins/Online/Experimental/OnlineServicesEOSGS/Source/Private/Online/UserFileEOSGS.cpp
--- a/Engine/Plugins/Online/Experimental/OnlineServicesEOSGS/Source/Private/Online/UserFileEOSGS.cpp
+++ b/Engine/Plugins/Online/Experimental/OnlineServicesEOSGS/Source/Private/Online/UserFileEOSGS.cpp
@@ -33,7 +33,7 @@ void FUserFileEOSGS::Initialize()
 void FUserFileEOSGS::UpdateConfig()
 {
 	Super::UpdateConfig();
-	Super::LoadConfig(Config);
+	LoadConfig(Config);
 }
 
 TOnlineAsyncOpHandle<FUserFileEnumerateFiles> FUserFileEOSGS::EnumerateFiles(FUserFileEnumerateFiles::Params&& InParams)
diff --git a/Engine/Plugins/Online/Experimental/OnlineServicesOSSAdapter/Source/Private/Online/UserFileOSSAdapter.cpp b/Engine/Plugins/Online/Experimental/OnlineServicesOSSAdapter/Source/Private/Online/UserFileOSSAdapter.cpp
--- a/Engine/Plugins/Online/Experimental/OnlineServicesOSSAdapter/Source/Private/Online/UserFileOSSAdapter.cpp
+++ b/Engine/Plugins/Online/Experimental/OnlineServicesOSSAdapter/Source/Private/Online/UserFileOSSAdapter.cpp
@@ -26,7 +26,7 @@ void FUserFileOSSAdapter::PostInitialize()
 void FUserFileOSSAdapter::UpdateConfig()
 {
 	Super::UpdateConfig();
-	Super::LoadConfig(Config);
+	LoadConfig(Config);
 }
 
 TOnlineAsyncOpHandle<FUserFileEnumerateFiles> FUserFileOSSAdapter::EnumerateFiles(FUserFileEnumerateFiles::Params&& InParams)
