Separation of session information in constant FSessionInfo and mutable FSessionSettings, and related refactorings.
Added IsPresenceSession method to interface.

diff --git a/Engine/Plugins/Online/Experimental/OnlineServices/Source/OnlineServicesCommon/Private/Online/NboSerializerCommonSvc.cpp b/Engine/Plugins/Online/Experimental/OnlineServices/Source/OnlineServicesCommon/Private/Online/NboSerializerCommonSvc.cpp
--- a/Engine/Plugins/Online/Experimental/OnlineServices/Source/OnlineServicesCommon/Private/Online/NboSerializerCommonSvc.cpp
+++ b/Engine/Plugins/Online/Experimental/OnlineServices/Source/OnlineServicesCommon/Private/Online/NboSerializerCommonSvc.cpp
@@ -63,32 +63,31 @@ void SerializeToBuffer(FNboSerializeToBuffer& Packet, const FSessionMember& Sess
 	SerializeToBuffer(Packet, SessionMember.MemberSettings);
 }
 
+void SerializeToBuffer(FNboSerializeToBuffer& Packet, const FSessionInfo& SessionInfo)
+{
+	Packet << SessionInfo.bAllowSanctionedPlayers;
+	Packet << SessionInfo.bAntiCheatProtected;
+	Packet << SessionInfo.bIsDedicatedServerSession;
+	Packet << SessionInfo.bIsLANSession;
+	Packet << SessionInfo.SessionIdOverride;
+	// SessionInfo.SessionId will be serialized in implementations as user id types will vary
+}
+
 void SerializeToBuffer(FNboSerializeToBuffer& Packet, const FSessionSettings& SessionSettings)
 {
 	Packet << SessionSettings.bAllowNewMembers;
-	Packet << SessionSettings.bAllowSanctionedPlayers;
-	Packet << SessionSettings.bAntiCheatProtected;
-	Packet << SessionSettings.bIsDedicatedServerSession;
-	Packet << SessionSettings.bIsLANSession;
-	Packet << SessionSettings.bPresenceEnabled;
 	SerializeToBuffer(Packet, SessionSettings.CustomSettings);
 	Packet << (uint8)SessionSettings.JoinPolicy;
-	Packet << SessionSettings.NumMaxPrivateConnections;;
-	Packet << SessionSettings.NumMaxPublicConnections;
-	Packet << SessionSettings.NumOpenPrivateConnections;
-	Packet << SessionSettings.NumOpenPublicConnections;
-	// SessionSettings.RegisteredPlayers will be serialized in implementations as user id types will vary
+	Packet << SessionSettings.NumMaxConnections;
 	Packet << SessionSettings.SchemaName;
-	Packet << SessionSettings.SessionIdOverride;
-	// SessionSettings.SessionMembers will be serialized in implementations as user id types will vary
 }
 
 void SerializeToBuffer(FNboSerializeToBuffer& Packet, const FSessionCommon& Session)
 {
 	// Session.OwnerUserId will be serialized in implementations as user id types will vary
-	// Session.SessionId will be serialized in implementations as user id types will vary
-	// Session.CurrentState won't be serialized as it should always be Valid
+	SerializeToBuffer(Packet, Session.SessionInfo);
 	SerializeToBuffer(Packet, Session.SessionSettings);
+	// Session.SessionMembers will be serialized in implementations as user id types will vary
 }
 
 /** SerializeFromBuffer methods */
@@ -163,44 +162,48 @@ void SerializeFromBuffer(FNboSerializeFromBuffer& Packet, FSessionMember& Sessio
 	SerializeFromBuffer(Packet, SessionMember.MemberSettings);
 }
 
-void SerializeFromBuffer(FNboSerializeFromBuffer& Packet, FSessionSettings& SessionSettings)
+void SerializeFromBuffer(FNboSerializeFromBuffer& Packet, FSessionInfo& SessionInfo)
 {
 	uint8 Read = 0;
 
 	// Read all the booleans as bytes
 	Packet >> Read;
-	SessionSettings.bAllowNewMembers = !!Read;
-	Packet >> Read;
-	SessionSettings.bAllowSanctionedPlayers = !!Read;
+	SessionInfo.bAllowSanctionedPlayers = !!Read;
 	Packet >> Read;
-	SessionSettings.bAntiCheatProtected = !!Read;
+	SessionInfo.bAntiCheatProtected = !!Read;
 	Packet >> Read;
-	SessionSettings.bIsDedicatedServerSession = !!Read;
+	SessionInfo.bIsDedicatedServerSession = !!Read;
 	Packet >> Read;
-	SessionSettings.bIsLANSession = !!Read;
+	SessionInfo.bIsLANSession = !!Read;
+
+	Packet >> SessionInfo.SessionIdOverride;
+
+	// SessionInfo.SessionId will be deserialized in implementations as user id types will vary
+}
+
+void SerializeFromBuffer(FNboSerializeFromBuffer& Packet, FSessionSettings& SessionSettings)
+{
+	uint8 Read = 0;
+
+	// Read all the booleans as bytes
 	Packet >> Read;
-	SessionSettings.bPresenceEnabled = !!Read;
+	SessionSettings.bAllowNewMembers = !!Read;
 
 	SerializeFromBuffer(Packet, SessionSettings.CustomSettings);
 
 	Packet >> Read;
 	SessionSettings.JoinPolicy = (ESessionJoinPolicy)Read;
 
-	Packet >> SessionSettings.NumMaxPrivateConnections;;
-	Packet >> SessionSettings.NumMaxPublicConnections;
-	Packet >> SessionSettings.NumOpenPrivateConnections;
-	Packet >> SessionSettings.NumOpenPublicConnections;
-	// SessionSettings.RegisteredPlayers will be deserialized in implementations as user id types will vary
+	Packet >> SessionSettings.NumMaxConnections;
 	Packet >> SessionSettings.SchemaName;
-	Packet >> SessionSettings.SessionIdOverride;
-	// SessionSettings.SessionMembers will be deserialized in implementations as user id types will vary
 }
 
 void SerializeFromBuffer(FNboSerializeFromBuffer& Packet, FSessionCommon& Session)
 {
 	// Session.OwnerUserId will be deserialized in implementations as user id types will vary
-	// Session.SessionId will be deserialized in implementations as user id types will vary
+	SerializeFromBuffer(Packet, Session.SessionInfo);
 	SerializeFromBuffer(Packet, Session.SessionSettings);
+	// Session.SessionMembers will be deserialized in implementations as user id types will vary
 }
 
 /* NboSerializerCommonSvc */ }
diff --git a/Engine/Plugins/Online/Experimental/OnlineServices/Source/OnlineServicesCommon/Private/Online/SessionsCommon.cpp b/Engine/Plugins/Online/Experimental/OnlineServices/Source/OnlineServicesCommon/Private/Online/SessionsCommon.cpp
--- a/Engine/Plugins/Online/Experimental/OnlineServices/Source/OnlineServicesCommon/Private/Online/SessionsCommon.cpp
+++ b/Engine/Plugins/Online/Experimental/OnlineServices/Source/OnlineServicesCommon/Private/Online/SessionsCommon.cpp
@@ -6,6 +6,53 @@
 #include "Online/OnlineServicesCommon.h"
 
 namespace UE::Online {
+	/** FSessionCommon */
+	FSessionCommon& FSessionCommon::operator+=(const FSessionUpdate& SessionUpdate)
+	{
+		if (SessionUpdate.OwnerAccountId.IsSet())
+		{
+			OwnerAccountId = SessionUpdate.OwnerAccountId.GetValue();
+		}
+
+		if (SessionUpdate.SessionSettingsChanges.IsSet())
+		{
+			SessionSettings += SessionUpdate.SessionSettingsChanges.GetValue();
+		}
+
+		// Session Members
+
+		for (const FAccountId& Key : SessionUpdate.RemovedSessionMembers)
+		{
+			SessionMembers.Remove(Key);
+		}
+
+		SessionMembers.Append(SessionUpdate.AddedSessionMembers);
+
+		for (const TPair<FAccountId, FSessionMemberChanges>& MemberEntry : SessionUpdate.SessionMembersChanges)
+		{
+			if (FSessionMember* SessionMember = SessionMembers.Find(MemberEntry.Key))
+			{
+				const FSessionMemberChanges& SessionMemberChanges = MemberEntry.Value;
+
+				for (const FName& Key : SessionMemberChanges.RemovedMemberSettings)
+				{
+					SessionMember->MemberSettings.Remove(Key);
+				}
+
+				SessionMember->MemberSettings.Append(SessionMemberChanges.AddedMemberSettings);
+
+				for (const TPair<FName, FCustomSessionSettingUpdate>& SettingEntry : SessionMemberChanges.ChangedMemberSettings)
+				{
+					if (FCustomSessionSetting* CustomSetting = SessionMember->MemberSettings.Find(SettingEntry.Key))
+					{
+						(*CustomSetting) = SettingEntry.Value.NewValue;
+					}
+				}
+			}
+		}
+
+		return *this;
+	}
 
 	/** FSessionsCommon */
 	FSessionsCommon::FSessionsCommon(FOnlineServicesCommon& InServices)
@@ -91,7 +138,19 @@ namespace UE::Online {
 		}
 		else
 		{
-			return  TOnlineResult<FGetPresenceSession>({ Errors::InvalidState() });
+			return TOnlineResult<FGetPresenceSession>({ Errors::InvalidState() });
+		}
+	}
+
+	TOnlineResult<FIsPresenceSession> FSessionsCommon::IsPresenceSession(FIsPresenceSession::Params&& Params) const
+	{
+		if (const FOnlineSessionId* PresenceSessionId = PresenceSessionsUserMap.Find(Params.LocalAccountId))
+		{
+			return TOnlineResult<FIsPresenceSession>(FIsPresenceSession::Result{ Params.SessionId == (*PresenceSessionId) });
+		}
+		else
+		{
+			return TOnlineResult<FIsPresenceSession>({ Errors::InvalidState() });
 		}
 	}
 
@@ -298,7 +357,7 @@ namespace UE::Online {
 
 	void FSessionsCommon::AddSessionInvite(const TSharedRef<FSessionInvite> SessionInvite, const TSharedRef<FSessionCommon> Session, const FAccountId& LocalAccountId)
 	{
-		AllSessionsById.Emplace(Session->SessionId, Session);
+		AllSessionsById.Emplace(Session->GetSessionId(), Session);
 
 		TMap<FSessionInviteId, TSharedRef<FSessionInvite>>& SessionInvitesMap = SessionInvitesUserMap.FindOrAdd(LocalAccountId);
 		SessionInvitesMap.Emplace(SessionInvite->InviteId, SessionInvite);
@@ -306,17 +365,17 @@ namespace UE::Online {
 
 	void FSessionsCommon::AddSearchResult(const TSharedRef<FSessionCommon> Session, const FAccountId& LocalAccountId)
 	{
-		AllSessionsById.Emplace(Session->SessionId, Session);
+		AllSessionsById.Emplace(Session->GetSessionId(), Session);
 
 		TArray<FOnlineSessionId>& SearchResults = SearchResultsUserMap.FindOrAdd(LocalAccountId);
-		SearchResults.Add(Session->SessionId);
+		SearchResults.Add(Session->GetSessionId());
 	}
 
 	void FSessionsCommon::AddSessionWithReferences(const TSharedRef<FSessionCommon> Session, const FName& SessionName, const FAccountId& LocalAccountId, bool bIsPresenceSession)
 	{
-		AllSessionsById.Emplace(Session->SessionId, Session);
+		AllSessionsById.Emplace(Session->GetSessionId(), Session);
 
-		AddSessionReferences(Session->SessionId, SessionName, LocalAccountId, bIsPresenceSession);
+		AddSessionReferences(Session->GetSessionId(), SessionName, LocalAccountId, bIsPresenceSession);
 	}
 
 	void FSessionsCommon::AddSessionReferences(const FOnlineSessionId SessionId, const FName& SessionName, const FAccountId& LocalAccountId, bool bIsPresenceSession)
@@ -351,13 +410,18 @@ namespace UE::Online {
 		}
 	}
 
-	void FSessionsCommon::ClearSessionReferences(const FOnlineSessionId SessionId, const FName& SessionName, const FAccountId& LocalAccountId, bool bIsPresenceSession)
+	void FSessionsCommon::ClearSessionReferences(const FOnlineSessionId SessionId, const FName& SessionName, const FAccountId& LocalAccountId)
 	{
 		NamedSessionUserMap.FindChecked(LocalAccountId).Remove(SessionName);
 
-		if (bIsPresenceSession)
+
+		TOnlineResult<FIsPresenceSession> IsPresenceSessionResult = IsPresenceSession({ LocalAccountId, SessionId });
+		if (IsPresenceSessionResult.IsOk())
 		{
-			ClearPresenceSession({ LocalAccountId });
+			if (IsPresenceSessionResult.GetOkValue().bIsPresenceSession)
+			{
+				ClearPresenceSession({ LocalAccountId });
+			}
 		}
 
 		ClearSessionByName(SessionName);
@@ -423,16 +487,12 @@ namespace UE::Online {
 
 		FSessionSettings& SessionSettings = FoundSession->SessionSettings;
 
-		if (SessionSettings.NumOpenPublicConnections > 0)
-		{
-			SessionSettings.NumOpenPublicConnections--;
-		}
-		else
+		if (FoundSession->GetNumOpenConnections() <= 0)
 		{
 			return TOnlineResult<FAddSessionMember>(Errors::InvalidState());
 		}
 
-		SessionSettings.SessionMembers.Emplace(Params.LocalAccountId, Params.NewSessionMember);
+		FoundSession->SessionMembers.Emplace(Params.LocalAccountId, Params.NewSessionMember);
 
 		return TOnlineResult<FAddSessionMember>();
 	}
@@ -448,16 +508,12 @@ namespace UE::Online {
 
 		FSessionSettings& SessionSettings = FoundSession->SessionSettings;
 		
-		if (SessionSettings.NumOpenPublicConnections < SessionSettings.NumMaxPublicConnections)
-		{
-			SessionSettings.NumOpenPublicConnections++;
-		}
-		else
+		if (FoundSession->GetNumOpenConnections() == SessionSettings.NumMaxConnections)
 		{
 			return TOnlineResult<FRemoveSessionMember>(Errors::InvalidState());
 		}
 
-		SessionSettings.SessionMembers.Remove(Params.LocalAccountId);
+		FoundSession->SessionMembers.Remove(Params.LocalAccountId);
 
 		return TOnlineResult<FRemoveSessionMember>();
 	}
@@ -478,9 +534,9 @@ namespace UE::Online {
 			return Errors::InvalidParams();
 		}
 
-		if (Params.SessionSettings.NumMaxPrivateConnections == 0 && Params.SessionSettings.NumMaxPublicConnections == 0)
+		if (Params.SessionSettings.NumMaxConnections == 0)
 		{
-			UE_LOG(LogTemp, Warning, TEXT("[FSessionsCommon::CheckCreateSessionParams] Could not create session with name [%s] with no valid NumMaxPrivateConnections [%d] or NumMaxPublicConnections [%d]"), *Params.SessionName.ToString(), Params.SessionSettings.NumMaxPrivateConnections, Params.SessionSettings.NumMaxPublicConnections);
+			UE_LOG(LogTemp, Warning, TEXT("[FSessionsCommon::CheckCreateSessionParams] Could not create session with name [%s] with no valid NumMaxConnections [%d]"), *Params.SessionName.ToString(), Params.SessionSettings.NumMaxConnections);
 
 			return Errors::InvalidParams();
 		}
@@ -498,14 +554,12 @@ namespace UE::Online {
 			return Errors::InvalidState();
 		}
 
-		if (Params.SessionSettings.bPresenceEnabled)
+		if (Params.bPresenceEnabled)
 		{
 			for (const TPair<FName, FOnlineSessionId>& Entry : LocalSessionsByName)
 			{
-				TOnlineResult<FGetSessionById> GetSessionByIdResult = GetSessionById({ Params.LocalAccountId, Entry.Value });
-				if (GetSessionByIdResult.IsOk())
-				{
-					if (GetSessionByIdResult.GetOkValue().Session->GetSessionSettings().bPresenceEnabled)
+				TOnlineResult<FGetPresenceSession> GetPresenceSessionResult = GetPresenceSession({ Params.LocalAccountId });
+				if (GetPresenceSessionResult.IsOk())
 				{
 					UE_LOG(LogTemp, Warning, TEXT("[FSessionsCommon::CheckCreateSessionState] Could not create session with bPresenceEnabled set to true when another already exists [%s]."), *Entry.Key.ToString());
 
@@ -513,7 +567,6 @@ namespace UE::Online {
 				}
 			}
 		}
-		}
 
 		// User login check for all local users
 		IAuthPtr Auth = Services.GetAuthInterface();
@@ -534,7 +587,7 @@ namespace UE::Online {
 		{
 			TSharedRef<const ISession> FoundSession = GetSessionByNameResult.GetOkValue().Session;
 
-			if (!FoundSession->GetSessionSettings().bIsDedicatedServerSession)
+			if (!FoundSession->GetSessionInfo().bIsDedicatedServerSession)
 			{
 				if (!Services.GetAuthInterface()->IsLoggedIn(Params.LocalAccountId))
 				{
@@ -589,7 +642,7 @@ namespace UE::Online {
 
 	FOnlineError FSessionsCommon::CheckStartMatchmakingParams(const FStartMatchmaking::Params& Params)
 	{
-		if (Params.SessionSettings.NumMaxPrivateConnections == 0 && Params.SessionSettings.NumMaxPublicConnections == 0)
+		if (Params.SessionCreationParameters.SessionSettings.NumMaxConnections == 0)
 		{
 			UE_LOG(LogTemp, Warning, TEXT("[FSessionsCommon::CheckStartMatchmakingParams] Could not create session with invalid NumMaxPrivateConnections and NumMaxPublicConnections"));
 
@@ -602,19 +655,19 @@ namespace UE::Online {
 	FOnlineError FSessionsCommon::CheckStartMatchmakingState(const FStartMatchmaking::Params& Params)
 	{
 		// Check if a session with that name already exists
-		TOptional<FOnlineError> Result = CheckSessionExistsByName(Params.LocalAccountId, Params.SessionName);
+		TOptional<FOnlineError> Result = CheckSessionExistsByName(Params.SessionCreationParameters.LocalAccountId, Params.SessionCreationParameters.SessionName);
 		if (!Result.IsSet()) // If CheckSessionExistsByName did not return an error, a session with that name already exists
 		{
-			UE_LOG(LogTemp, Warning, TEXT("[FSessionsCommon::CheckStartMatchmakingState] Could not join session with name [%s]. A session with that name already exists"), *Params.SessionName.ToString());
+			UE_LOG(LogTemp, Warning, TEXT("[FSessionsCommon::CheckStartMatchmakingState] Could not join session with name [%s]. A session with that name already exists"), *Params.SessionCreationParameters.SessionName.ToString());
 
 			return Errors::InvalidState(); // TODO: New error: Session with name %s already exists
 		}
 
 		// User login check for all local users
 		IAuthPtr Auth = Services.GetAuthInterface();
-		if (!Auth->IsLoggedIn(Params.LocalAccountId))
+		if (!Auth->IsLoggedIn(Params.SessionCreationParameters.LocalAccountId))
 		{
-			UE_LOG(LogTemp, Warning, TEXT("[FSessionsCommon::CheckStartMatchmakingState] Could not join session with user [%s] not logged in"), *ToLogString(Params.LocalAccountId));
+			UE_LOG(LogTemp, Warning, TEXT("[FSessionsCommon::CheckStartMatchmakingState] Could not join session with user [%s] not logged in"), *ToLogString(Params.SessionCreationParameters.LocalAccountId));
 
 			return Errors::InvalidUser();
 		}
@@ -683,32 +736,32 @@ namespace UE::Online {
 
 		const FSessionSettings& SessionSettings = FoundSession->GetSessionSettings();
 
-		if (SessionSettings.SessionMembers.Contains(Params.LocalAccountId))
+		if (FoundSession->GetSessionMembers().Contains(Params.LocalAccountId))
 		{
 			UE_LOG(LogTemp, Warning, TEXT("[FSessionsCommon::CheckJoinSessionState] Could not join session. User [%s] already in session"), *ToLogString(Params.LocalAccountId));
 
 			return Errors::AccessDenied();
 		}
 
-		if (!SessionSettings.bAllowNewMembers)
-		{
-			UE_LOG(LogTemp, Warning, TEXT("[FSessionsCommon::CheckJoinSessionState] Could not join session. bAllowNewMembers is set to false"));
-
-			return Errors::AccessDenied();
-		}
-
-		if (SessionSettings.NumMaxPublicConnections > 0 && SessionSettings.NumOpenPublicConnections > 0)
+		if (!FoundSession->IsJoinable())
 		{
-			UE_LOG(LogTemp, Warning, TEXT("[FSessionsCommon::CheckJoinSessionState] Could not join session. No open public connections"));
+			UE_LOG(LogTemp, Warning, TEXT("[FSessionsCommon::CheckJoinSessionState] Could not join session. Session not joinable "));
 
 			return Errors::AccessDenied();
 		}
 
-		if (SessionSettings.NumMaxPrivateConnections > 0 && SessionSettings.NumOpenPrivateConnections > 0)
+		if (Params.bPresenceEnabled)
 		{
-			UE_LOG(LogTemp, Warning, TEXT("[FSessionsCommon::CheckJoinSessionState] Could not join session. No open private connections"));
+			for (const TPair<FName, FOnlineSessionId>& Entry : LocalSessionsByName)
+			{
+				TOnlineResult<FGetPresenceSession> GetPresenceSessionResult = GetPresenceSession({ Params.LocalAccountId });
+				if (GetPresenceSessionResult.IsOk())
+				{
+					UE_LOG(LogTemp, Warning, TEXT("[FSessionsCommon::CheckJoinSessionState] Could not create session with bPresenceEnabled set to true when another already exists [%s]."), *Entry.Key.ToString());
 
-			return Errors::AccessDenied();
+					return Errors::InvalidState();
+				}
+			}
 		}
 
 		return Errors::Success();
@@ -770,7 +823,7 @@ namespace UE::Online {
 		{
 			TSharedRef<const ISession> FoundSession = GetSessionByNameResult.GetOkValue().Session;
 
-			if (!FoundSession->GetSessionSettings().bIsDedicatedServerSession)
+			if (!FoundSession->GetSessionInfo().bIsDedicatedServerSession)
 			{
 				if (!Services.GetAuthInterface()->IsLoggedIn(Params.LocalAccountId))
 				{
diff --git a/Engine/Plugins/Online/Experimental/OnlineServices/Source/OnlineServicesCommon/Public/Online/SessionsCommon.h b/Engine/Plugins/Online/Experimental/OnlineServices/Source/OnlineServicesCommon/Public/Online/SessionsCommon.h
--- a/Engine/Plugins/Online/Experimental/OnlineServices/Source/OnlineServicesCommon/Public/Online/SessionsCommon.h
+++ b/Engine/Plugins/Online/Experimental/OnlineServices/Source/OnlineServicesCommon/Public/Online/SessionsCommon.h
@@ -136,20 +136,32 @@ public:
 
 	// ISession
 	virtual const FAccountId GetOwnerAccountId() const override				{ return OwnerAccountId; }
-	virtual const FOnlineSessionId GetSessionId() const override		{ return SessionId; }
+	virtual const FOnlineSessionId GetSessionId() const override			{ return GetSessionInfo().SessionId; }
+	virtual const uint32 GetNumOpenConnections() const override				{ return SessionSettings.NumMaxConnections - SessionMembers.Num(); }
+	virtual const FSessionInfo& GetSessionInfo() const override				{ return SessionInfo; }
 	virtual const FSessionSettings GetSessionSettings() const override		{ return SessionSettings; }
+	virtual const FSessionMembersMap& GetSessionMembers() const override	{ return SessionMembers; }
+
+	virtual bool IsJoinable() const override								{ return GetNumOpenConnections() > 0 && SessionSettings.bAllowNewMembers; }
 
 	virtual FString ToLogString() const override							{ return FString(); } // TODO: Implement after completing refactor
 
 public:
+	/** Session information that will remain constant throughout the session's lifetime */
+	FSessionInfo SessionInfo;
+	
+	/** The following members can be updated, and such update will be transmitted via a FSessionUpdated event */
+
 	/** The user who currently owns the session */
 	FAccountId OwnerAccountId;
 
-	/** The id for the session, platform dependent */
-	FOnlineSessionId SessionId;
-
 	/** Set of session properties that can be altered by the session owner */
 	FSessionSettings SessionSettings;
+
+	/* Map of session member ids to their corresponding user-defined properties */
+	FSessionMembersMap SessionMembers;
+
+	FSessionCommon& operator+=(const FSessionUpdate& SessionUpdate);
 };
 
 struct FGetMutableSessionByName
@@ -202,6 +214,7 @@ public:
 	virtual TOnlineResult<FGetSessionByName> GetSessionByName(FGetSessionByName::Params&& Params) const override;
 	virtual TOnlineResult<FGetSessionById> GetSessionById(FGetSessionById::Params&& Params) const override;
 	virtual TOnlineResult<FGetPresenceSession> GetPresenceSession(FGetPresenceSession::Params&& Params) const override;
+	virtual TOnlineResult<FIsPresenceSession> IsPresenceSession(FIsPresenceSession::Params&& Params) const override;
 	virtual TOnlineResult<FSetPresenceSession> SetPresenceSession(FSetPresenceSession::Params&& Params) override;
 	virtual TOnlineResult<FClearPresenceSession> ClearPresenceSession(FClearPresenceSession::Params&& Params) override;
 	virtual TOnlineAsyncOpHandle<FCreateSession> CreateSession(FCreateSession::Params&& Params) override;
@@ -232,7 +245,7 @@ protected:
 	void AddSessionReferences(const FOnlineSessionId SessionId, const FName& SessionName, const FAccountId& LocalAccountId, bool bIsPresenceSession);
 
 	void ClearSessionInvitesForSession(const FAccountId& LocalAccountId, const FOnlineSessionId SessionId);
-	void ClearSessionReferences(const FOnlineSessionId SessionId, const FName& SessionName, const FAccountId& LocalAccountId, bool bIsPresenceSession);
+	void ClearSessionReferences(const FOnlineSessionId SessionId, const FName& SessionName, const FAccountId& LocalAccountId);
 
 	TOnlineResult<FAddSessionMember> AddSessionMemberImpl(const FAddSessionMember::Params& Params);
 	TOnlineResult<FRemoveSessionMember> RemoveSessionMemberImpl(const FRemoveSessionMember::Params& Params);
diff --git a/Engine/Plugins/Online/Experimental/OnlineServices/Source/OnlineServicesCommonEngineUtils/Private/Online/SessionsLAN.cpp b/Engine/Plugins/Online/Experimental/OnlineServices/Source/OnlineServicesCommonEngineUtils/Private/Online/SessionsLAN.cpp
--- a/Engine/Plugins/Online/Experimental/OnlineServices/Source/OnlineServicesCommonEngineUtils/Private/Online/SessionsLAN.cpp
+++ b/Engine/Plugins/Online/Experimental/OnlineServices/Source/OnlineServicesCommonEngineUtils/Private/Online/SessionsLAN.cpp
@@ -153,14 +153,14 @@ TOnlineAsyncOpHandle<FCreateSession> FSessionsLAN::CreateSession(FCreateSession:
 
 		TSharedRef<FSessionLAN> NewSessionLANRef = MakeShared<FSessionLAN>();
 		NewSessionLANRef->OwnerAccountId = OpParams.LocalAccountId;
-		NewSessionLANRef->SessionId = FOnlineSessionIdRegistryLAN::GetChecked(Services.GetServicesProvider()).GetNextSessionId();
+		NewSessionLANRef->SessionInfo.SessionId = FOnlineSessionIdRegistryLAN::GetChecked(Services.GetServicesProvider()).GetNextSessionId();
 		NewSessionLANRef->SessionSettings = OpParams.SessionSettings;
 
 		// For LAN sessions, we'll add the Session member manually instead of calling JoinSession since there is no API calls involved
-		NewSessionLANRef->SessionSettings.SessionMembers.Emplace(OpParams.LocalAccountId, OpParams.SessionMemberData);
+		NewSessionLANRef->SessionMembers.Emplace(OpParams.LocalAccountId, OpParams.SessionMemberData);
 
 		// We save the local object for the session, and set up the appropriate references
-		AddSessionWithReferences(NewSessionLANRef, OpParams.SessionName, OpParams.LocalAccountId, NewSessionLANRef->SessionSettings.bPresenceEnabled);
+		AddSessionWithReferences(NewSessionLANRef, OpParams.SessionName, OpParams.LocalAccountId, OpParams.bPresenceEnabled);
 
 		Op.SetResult({ });
 	})
@@ -211,13 +211,15 @@ TOnlineAsyncOpHandle<FUpdateSession> FSessionsLAN::UpdateSession(FUpdateSession:
 			}
 		}
 
-		SessionSettings += UpdatedSettings;
+		//SessionSettings += UpdatedSettings;
 
 		// We set the result and fire the event
 		Op.SetResult({ });
 
-		FSessionUpdated SessionUpdatedEvent{ OpParams.SessionName, UpdatedSettings };
-		SessionEvents.OnSessionUpdated.Broadcast(SessionUpdatedEvent);
+		//FSessionUpdated SessionUpdatedEvent{ OpParams.SessionName, UpdatedSettings };
+		//SessionEvents.OnSessionUpdated.Broadcast(SessionUpdatedEvent);
+
+		// TODO: Will be refactored as part of the UpdateSettings split
 	})
 	.Enqueue(GetSerialQueue());
 
@@ -296,7 +298,7 @@ TOnlineAsyncOpHandle<FJoinSession> FSessionsLAN::JoinSession(FJoinSession::Param
 		const TSharedRef<const ISession>& FoundSession = GetSessionByIdResult.GetOkValue().Session;
 
 		// We set up the appropriate references for the session
-		AddSessionReferences(FoundSession->GetSessionId(), OpParams.SessionName, OpParams.LocalAccountId, FoundSession->GetSessionSettings().bPresenceEnabled);
+		AddSessionReferences(FoundSession->GetSessionId(), OpParams.SessionName, OpParams.LocalAccountId, OpParams.bPresenceEnabled);
 
 		Op.SetResult(FJoinSession::Result{ });
 
@@ -334,7 +336,7 @@ TOnlineAsyncOpHandle<FLeaveSession> FSessionsLAN::LeaveSession(FLeaveSession::Pa
 				StopLANSession();
 			}
 
-			ClearSessionReferences(FoundSession->GetSessionId(), OpParams.SessionName, OpParams.LocalAccountId, FoundSession->GetSessionSettings().bPresenceEnabled);
+			ClearSessionReferences(FoundSession->GetSessionId(), OpParams.SessionName, OpParams.LocalAccountId);
 		}
 
 		Op.SetResult(FLeaveSession::Result{ });
diff --git a/Engine/Plugins/Online/Experimental/OnlineServices/Source/OnlineServicesInterface/Private/Online/Sessions.cpp b/Engine/Plugins/Online/Experimental/OnlineServices/Source/OnlineServicesInterface/Private/Online/Sessions.cpp
--- a/Engine/Plugins/Online/Experimental/OnlineServices/Source/OnlineServicesInterface/Private/Online/Sessions.cpp
+++ b/Engine/Plugins/Online/Experimental/OnlineServices/Source/OnlineServicesInterface/Private/Online/Sessions.cpp
@@ -37,50 +37,30 @@ void LexFromString(ESessionJoinPolicy& Value, const TCHAR* InStr)
 }
 
 #define COPY_TOPTIONAL_VALUE_IF_SET(Value) \
-	if (UpdatedSettings.Value.IsSet()) \
+	if (SettingsChanges.Value.IsSet()) \
 	{ \
-		Value = UpdatedSettings.Value.GetValue(); \
+		Value = SettingsChanges.Value.GetValue(); \
 	} \
 
-FSessionSettings& FSessionSettings::operator+=(const FSessionSettingsUpdate& UpdatedSettings)
+FSessionSettings& FSessionSettings::operator+=(const FSessionSettingsChanges& SettingsChanges)
 {
 	COPY_TOPTIONAL_VALUE_IF_SET(SchemaName) // TODO: We may need some additional logic for schema changes
-	COPY_TOPTIONAL_VALUE_IF_SET(NumMaxPublicConnections)
-	COPY_TOPTIONAL_VALUE_IF_SET(NumOpenPublicConnections)
-	COPY_TOPTIONAL_VALUE_IF_SET(NumMaxPrivateConnections)
-	COPY_TOPTIONAL_VALUE_IF_SET(NumOpenPrivateConnections)
+	COPY_TOPTIONAL_VALUE_IF_SET(NumMaxConnections)
 	COPY_TOPTIONAL_VALUE_IF_SET(JoinPolicy)
-	COPY_TOPTIONAL_VALUE_IF_SET(SessionIdOverride)
-	COPY_TOPTIONAL_VALUE_IF_SET(bIsDedicatedServerSession)
 	COPY_TOPTIONAL_VALUE_IF_SET(bAllowNewMembers)
-	COPY_TOPTIONAL_VALUE_IF_SET(bAllowSanctionedPlayers)
-	COPY_TOPTIONAL_VALUE_IF_SET(bAntiCheatProtected)
-	COPY_TOPTIONAL_VALUE_IF_SET(bPresenceEnabled)
 		
-	for (const FName& Key : UpdatedSettings.RemovedCustomSettings)
+	for (const FName& Key : SettingsChanges.RemovedCustomSettings)
 	{
 		CustomSettings.Remove(Key);
 	}
 
-	CustomSettings.Append(UpdatedSettings.UpdatedCustomSettings);
+	CustomSettings.Append(SettingsChanges.AddedCustomSettings);
 
-	for (const FAccountId& Key : UpdatedSettings.RemovedSessionMembers)
+	for (const TPair<FName, FCustomSessionSettingUpdate>& SettingEntry : SettingsChanges.ChangedCustomSettings)
 	{
-		SessionMembers.Remove(Key);
-	}
-
-	for (const TPair<FAccountId, FSessionMemberUpdate>& Entry : UpdatedSettings.UpdatedSessionMembers)
-	{
-		if (FSessionMember* SessionMember = SessionMembers.Find(Entry.Key))
+		if (FCustomSessionSetting* CustomSetting = CustomSettings.Find(SettingEntry.Key))
 		{
-			const FSessionMemberUpdate& SessionMemberUpdate = Entry.Value;
-
-			for (const FName& Key : SessionMemberUpdate.RemovedMemberSettings)
-			{
-				SessionMember->MemberSettings.Remove(Key);
-			}
-
-			SessionMember->MemberSettings.Append(SessionMemberUpdate.UpdatedMemberSettings);
+			(*CustomSetting) = SettingEntry.Value.NewValue;
 		}
 	}
 
@@ -117,17 +97,9 @@ FSessionMemberUpdate& FSessionMemberUpdate::operator+=(FSessionMemberUpdate&& Up
 FSessionSettingsUpdate& FSessionSettingsUpdate::operator+=(FSessionSettingsUpdate&& UpdatedValue)
 {
 	MOVE_TOPTIONAL_IF_SET(SchemaName)
-	MOVE_TOPTIONAL_IF_SET(NumMaxPublicConnections)
-	MOVE_TOPTIONAL_IF_SET(NumOpenPublicConnections)
-	MOVE_TOPTIONAL_IF_SET(NumMaxPrivateConnections)
-	MOVE_TOPTIONAL_IF_SET(NumOpenPrivateConnections)
+	MOVE_TOPTIONAL_IF_SET(NumMaxConnections)
 	MOVE_TOPTIONAL_IF_SET(JoinPolicy)
-	MOVE_TOPTIONAL_IF_SET(SessionIdOverride)
-	MOVE_TOPTIONAL_IF_SET(bIsDedicatedServerSession)
 	MOVE_TOPTIONAL_IF_SET(bAllowNewMembers)
-	MOVE_TOPTIONAL_IF_SET(bAllowSanctionedPlayers)
-	MOVE_TOPTIONAL_IF_SET(bAntiCheatProtected)
-	MOVE_TOPTIONAL_IF_SET(bPresenceEnabled)
 
 	for (const TPair<FName, FCustomSessionSetting>& UpdatedCustomSetting : UpdatedValue.UpdatedCustomSettings)
 	{
diff --git a/Engine/Plugins/Online/Experimental/OnlineServices/Source/OnlineServicesInterface/Public/Online/Sessions.h b/Engine/Plugins/Online/Experimental/OnlineServices/Source/OnlineServicesInterface/Public/Online/Sessions.h
--- a/Engine/Plugins/Online/Experimental/OnlineServices/Source/OnlineServicesInterface/Public/Online/Sessions.h
+++ b/Engine/Plugins/Online/Experimental/OnlineServices/Source/OnlineServicesInterface/Public/Online/Sessions.h
@@ -36,6 +36,15 @@ struct FCustomSessionSetting
 
 using FCustomSessionSettingsMap = TMap<FSchemaAttributeId, FCustomSessionSetting>;
 
+struct FCustomSessionSettingUpdate
+{
+	FCustomSessionSetting OldValue;
+
+	FCustomSessionSetting NewValue;
+};
+
+using FCustomSessionSettingUpdateMap = TMap<FName, FCustomSessionSettingUpdate>;
+
 /** A member is a player that is part of the session, and it stops being a member when they leave it */
 struct FSessionMember
 {
@@ -63,32 +72,17 @@ enum class ESessionJoinPolicy : uint8
 ONLINESERVICESINTERFACE_API const TCHAR* LexToString(ESessionJoinPolicy Value);
 ONLINESERVICESINTERFACE_API void LexFromString(ESessionJoinPolicy& Value, const TCHAR* InStr);
 
+/** Contains new values for an FSessions modifiable settings. Taken as a parameter by FUpdateSessions method */
 struct FSessionSettingsUpdate
 {
 	/** Set with an updated value if the SchemaName field will be changed in the update operation */
 	TOptional<FSchemaId> SchemaName;
-	/** Set with an updated value if the NumMaxPublicConnections field will be changed in the update operation */
-	TOptional<uint32> NumMaxPublicConnections;
-	/** Set with an updated value if the NumOpenPublicConnections field will be changed in the update operation */
-	TOptional<uint32> NumOpenPublicConnections;
-	/** Set with an updated value if the NumMaxPrivateConnections field will be changed in the update operation */
-	TOptional<uint32> NumMaxPrivateConnections;
-	/** Set with an updated value if the NumOpenPrivateConnections field will be changed in the update operation */
-	TOptional<uint32> NumOpenPrivateConnections;
+	/** Set with an updated value if the NumMaxConnections field will be changed in the update operation */
+	TOptional<uint32> NumMaxConnections;
 	/** Set with an updated value if the JoinPolicy field will be changed in the update operation */
 	TOptional<ESessionJoinPolicy> JoinPolicy;
-	/** Set with an updated value if the SessionIdOverride field will be changed in the update operation */
-	TOptional<FString> SessionIdOverride;
-	/** Set with an updated value if the bIsDedicatedServerSession field will be changed in the update operation */
-	TOptional<bool> bIsDedicatedServerSession;
 	/** Set with an updated value if the bAllowNewMembers field will be changed in the update operation */
 	TOptional<bool> bAllowNewMembers;
-	/** Set with an updated value if the bAllowSanctionedPlayers field will be changed in the update operation */
-	TOptional<bool> bAllowSanctionedPlayers;
-	/** Set with an updated value if the bAntiCheatProtected field will be changed in the update operation */
-	TOptional<bool> bAntiCheatProtected;
-	/** Set with an updated value if the bPresenceEnabled field will be changed in the update operation */
-	TOptional<bool> bPresenceEnabled;
 
 	/** Updated values for custom settings to change in the update operation*/
 	FCustomSessionSettingsMap UpdatedCustomSettings;
@@ -103,25 +97,67 @@ struct FSessionSettingsUpdate
 	FSessionSettingsUpdate& operator+=(FSessionSettingsUpdate&& UpdatedValue);
 };
 
+/** Contains updated data for any modifiable members of FSessionSettings. Member of FSessionUpdated event */
+struct FSessionSettingsChanges
+{
+	/* If set, the FSessionSettings's SchemaName member will be updated to this value */
+	TOptional<FName> SchemaName;
+	/** If set, the FSessionSettings's NumMaxConnections member will be updated to this value */
+	TOptional<uint32> NumMaxConnections;
+	/** If set, the FSessionSettings's JoinPolicy member will be updated to this value */
+	TOptional<ESessionJoinPolicy> JoinPolicy;
+	/** If set, the FSessionSettings's bAllowNewMembers member will be updated to this value */
+	TOptional<bool> bAllowNewMembers;
+
+	/** New custom settings, with their values */
+	FCustomSessionSettingsMap AddedCustomSettings;
+
+	/** Existing custom settings that changed value, including new and old values */
+	FCustomSessionSettingUpdateMap ChangedCustomSettings;
+
+	/** Keys for removed custom settings */
+	TArray<FName> RemovedCustomSettings;
+};
+
+/** Contains updated data for any modifiable members of FSessionMember. Member of FSessionUpdated event */
+struct FSessionMemberChanges
+{
+	/** New custom settings, with their values */
+	FCustomSessionSettingsMap AddedMemberSettings;
+
+	/** Existing custom settings that changed value, including new and old values */
+	FCustomSessionSettingUpdateMap ChangedMemberSettings;
+
+	/** Keys for removed custom settings */
+	TArray<FName> RemovedMemberSettings;
+};
+
+/** Set of all of an FSession's defining properties that can be updated by the session owner during its lifetime */
 struct ONLINESERVICESINTERFACE_API FSessionSettings
 {
 	/* The schema which will be applied to the session */
-	FSchemaId SchemaName;
+	FName SchemaName;
 
-	/* Maximum number of public slots for session members */
-	uint32 NumMaxPublicConnections = 0;
+	/* Maximum number of slots for session members */
+	uint32 NumMaxConnections = 0;
 
-	/* Number of available public slots for session members */
-	uint32 NumOpenPublicConnections = 0;
+	/* Enum value describing the level of restriction to join the session. Public by default */
+	ESessionJoinPolicy JoinPolicy = ESessionJoinPolicy::Public;;
 
-	/* Maximum number of private slots for session members */
-	uint32 NumMaxPrivateConnections = 0;
+	/* Override value to restrict the session from accepting new members, regardless of other factors. True by default */
+	bool bAllowNewMembers = true;
 
-	/* Number of available private slots for session members */
-	uint32 NumOpenPrivateConnections = 0;
+	/* Map of user-defined settings to be passed to the platform APIs as additional information for various purposes */
+	FCustomSessionSettingsMap CustomSettings;
 
-	/* Enum value describing the level of restriction to join the session */
-	ESessionJoinPolicy JoinPolicy = ESessionJoinPolicy::Public;
+	FSessionSettings& operator+=(const FSessionSettingsChanges& UpdatedValue);
+};
+
+/** Information about an FSession that will be set at creation time and remain constant during its lifetime */
+struct FSessionInfo
+{
+	/** The id for the session, platform dependent */
+	FOnlineSessionId SessionId;
 
 	/* In platforms that support this feature, it will set the session id to this value. Might be subject to minimum and maximum length */
 	FString SessionIdOverride;
@@ -132,25 +168,11 @@ struct ONLINESERVICESINTERFACE_API FSessionSettings
 	/* Whether the session is configured to run as a dedicated server. Only available in some platforms. False by default */
 	bool bIsDedicatedServerSession = false;
 
-	/* Whether players are accepted as new members in the session. Can vary depending on various factors, like the number of free slots available, or Join-In-Progress preferences when the session has started. True by default */
-	bool bAllowNewMembers = true;
-
 	/* Whether this session will allow sanctioned players to join it. True by default */
 	bool bAllowSanctionedPlayers = true;
 
-	/*Whether this is a secure session protected by anti-cheat services. False by default */
+	/* Whether this is a secure session protected by anti-cheat services. False by default */
 	bool bAntiCheatProtected = false;
-
-	/* Whether this session will show its information in presence updates. Can only be set in one session at a time. False by default */
-	bool bPresenceEnabled = false;
-  
- 	/* Map of user-defined settings to be passed to the platform APIs as additional information for various purposes */
- 	FCustomSessionSettingsMap CustomSettings;
- 
- 	/* Map of session member ids to their corresponding user-defined settings */
- 	FSessionMembersMap SessionMembers;
-
-	FSessionSettings& operator+=(const FSessionSettingsUpdate& UpdatedValue);
 };
 
 class ISession
@@ -158,7 +180,14 @@ class ISession
 public:
 	virtual const FAccountId GetOwnerAccountId() const = 0;
 	virtual const FOnlineSessionId GetSessionId() const = 0;
+	virtual const uint32 GetNumOpenConnections() const = 0;
+	virtual const FSessionInfo& GetSessionInfo() const = 0;
 	virtual const FSessionSettings GetSessionSettings() const = 0;
+	virtual const FSessionMembersMap& GetSessionMembers() const = 0;
+
+	/** Evaluates a series of factors to determine if a session is accepting new members */
+	virtual bool IsJoinable() const = 0;
+
 	virtual FString ToLogString() const = 0;
 };
 ONLINESERVICESINTERFACE_API const FString ToLogString(const ISession& Session);
@@ -246,6 +275,23 @@ struct FGetPresenceSession
 	};
 };
 
+struct FIsPresenceSession
+{
+	static constexpr TCHAR Name[] = TEXT("IsPresenceSession");
+
+	struct Params
+	{
+		FAccountId LocalAccountId;
+
+		FOnlineSessionId SessionId;
+	};
+
+	struct Result
+	{
+		bool bIsPresenceSession;
+	};
+};
+
 struct FSetPresenceSession
 {
 	static constexpr TCHAR Name[] = TEXT("SetPresenceSession");
@@ -291,6 +337,24 @@ struct FCreateSession
 		/** Information for the local user who will join the session after creation */
 		FSessionMember SessionMemberData;
 
+		/* In platforms that support this feature, it will set the session id to this value. Might be subject to minimum and maximum length */
+		FString SessionIdOverride;
+
+		/** Whether this session should be set as the user's new presence session. False by default */
+		bool bPresenceEnabled = false;
+
+		/* Whether the session is only available in the local network and not via internet connection. Only available in some platforms. False by default */
+		bool bIsLANSession = false;
+
+		/* Whether the session is configured to run as a dedicated server. Only available in some platforms. False by default */
+		bool bIsDedicatedServerSession = false;
+
+		/* Whether this session will allow sanctioned players to join it. True by default */
+		bool bAllowSanctionedPlayers = true;
+
+		/* Whether this is a secure session protected by anti-cheat services. False by default */
+		bool bAntiCheatProtected = false;
+
 		/** Settings object to define session properties during creation */
 		FSessionSettings SessionSettings;
 	};
@@ -382,20 +446,11 @@ struct FStartMatchmaking
 
 	struct Params
 	{
-		/** The local user agent which will perform the action. */
-		FAccountId LocalAccountId;
-
-		/* Local name for the session */
-		FName SessionName;
-
-		/* Information for the local user who will create or join the session */
-		FSessionMember SessionMemberData;
-
-		/* Preferred settings to be used during session creation */
-		FSessionSettings SessionSettings;
+		/* Session creation parameters */
+		FCreateSession::Params SessionCreationParameters;
 
 		/* Filters to apply when searching for sessions */
-		TArray<FFindSessionsSearchFilter> SearchFilters;
+		TArray<FFindSessionsSearchFilter> SessionSearchFilters;
 	};
 
 	struct Result
@@ -421,6 +476,9 @@ struct FJoinSession
 
 		/* Information for the local user who will join the session */
 		FSessionMember SessionMemberData;
+
+		/* Whether this session should be set as the user's new presence session. False by default */
+		bool bPresenceEnabled = false;
 	};
 
 	struct Result
@@ -544,13 +602,36 @@ struct FSessionLeft
 	FAccountId LocalAccountId;
 };
 
+using FSessionMemberChangesMap = TMap<FAccountId, FSessionMemberChanges>;
+
+/** Contains updated data for any modifiable members of ISession */
+struct ONLINESERVICESINTERFACE_API FSessionUpdate
+{
+	/** If set, the OwnerUserId member will have updated to this value */
+	TOptional<FAccountId> OwnerAccountId;
+
+	/** If set, the SessionSettings member will have updated using the struct information */
+	TOptional<FSessionSettingsChanges> SessionSettingsChanges;
+
+	/** Session member information for members that just joined the session */
+	FSessionMembersMap AddedSessionMembers;
+
+	/** Updated values for session member information */
+	FSessionMemberChangesMap SessionMembersChanges;
+
+	/** Id handles for members that just left the session */
+	TArray<FAccountId> RemovedSessionMembers;
+
+	FSessionUpdate& operator+=(const FSessionUpdate& SessionUpdate);
+};
+
 struct FSessionUpdated
 {
-	/* Local name for the updated session object */
+	/* Name for the session updated */
 	FName SessionName;
 
 	/* Updated session settings */
-	FSessionSettingsUpdate UpdatedSettings;
+	FSessionUpdate SessionUpdate;
 };
 
 struct FSessionInviteReceived
@@ -624,6 +705,14 @@ public:
 	 */
 	virtual TOnlineResult<FGetPresenceSession> GetPresenceSession(FGetPresenceSession::Params&& Params) const = 0;
 
+	/**
+	 * Returns whether the session with the given id is set as the presence session for the user.
+	 *
+	 * @params Parameters for the IsPresenceSession call
+	 * return
+	 */
+	virtual TOnlineResult<FIsPresenceSession> IsPresenceSession(FIsPresenceSession::Params&& Params) const = 0;
+
 	/**
 	 * Sets the session with the given id as the presence session for the user.
 	 *
@@ -789,6 +878,11 @@ BEGIN_ONLINE_STRUCT_META(FCustomSessionSetting)
 	ONLINE_STRUCT_FIELD(FCustomSessionSetting, ID)
 END_ONLINE_STRUCT_META()
 
+BEGIN_ONLINE_STRUCT_META(FCustomSessionSettingUpdate)
+	ONLINE_STRUCT_FIELD(FCustomSessionSettingUpdate, OldValue),
+	ONLINE_STRUCT_FIELD(FCustomSessionSettingUpdate, NewValue)
+END_ONLINE_STRUCT_META()
+
 BEGIN_ONLINE_STRUCT_META(FSessionMember)
 	ONLINE_STRUCT_FIELD(FSessionMember, MemberSettings)
 END_ONLINE_STRUCT_META()
@@ -800,35 +894,26 @@ END_ONLINE_STRUCT_META()
 
 BEGIN_ONLINE_STRUCT_META(FSessionSettings)
 	ONLINE_STRUCT_FIELD(FSessionSettings, SchemaName),
-	ONLINE_STRUCT_FIELD(FSessionSettings, NumMaxPublicConnections),
-	ONLINE_STRUCT_FIELD(FSessionSettings, NumOpenPublicConnections),
-	ONLINE_STRUCT_FIELD(FSessionSettings, NumMaxPrivateConnections),
-	ONLINE_STRUCT_FIELD(FSessionSettings, NumOpenPrivateConnections),
+	ONLINE_STRUCT_FIELD(FSessionSettings, NumMaxConnections),
 	ONLINE_STRUCT_FIELD(FSessionSettings, JoinPolicy),
-	ONLINE_STRUCT_FIELD(FSessionSettings, SessionIdOverride),
-	ONLINE_STRUCT_FIELD(FSessionSettings, bIsLANSession),
-	ONLINE_STRUCT_FIELD(FSessionSettings, bIsDedicatedServerSession),
 	ONLINE_STRUCT_FIELD(FSessionSettings, bAllowNewMembers),
-	ONLINE_STRUCT_FIELD(FSessionSettings, bAllowSanctionedPlayers),
-	ONLINE_STRUCT_FIELD(FSessionSettings, bAntiCheatProtected),
-	ONLINE_STRUCT_FIELD(FSessionSettings, bPresenceEnabled),
-	ONLINE_STRUCT_FIELD(FSessionSettings, CustomSettings),
-	ONLINE_STRUCT_FIELD(FSessionSettings, SessionMembers)
+	ONLINE_STRUCT_FIELD(FSessionSettings, CustomSettings)
+END_ONLINE_STRUCT_META()
+
+BEGIN_ONLINE_STRUCT_META(FSessionInfo)
+	ONLINE_STRUCT_FIELD(FSessionInfo, SessionId),
+	ONLINE_STRUCT_FIELD(FSessionInfo, SessionIdOverride),
+	ONLINE_STRUCT_FIELD(FSessionInfo, bIsLANSession),
+	ONLINE_STRUCT_FIELD(FSessionInfo, bIsDedicatedServerSession),
+	ONLINE_STRUCT_FIELD(FSessionInfo, bAllowSanctionedPlayers),
+	ONLINE_STRUCT_FIELD(FSessionInfo, bAntiCheatProtected)
 END_ONLINE_STRUCT_META()
 
 BEGIN_ONLINE_STRUCT_META(FSessionSettingsUpdate)
 	ONLINE_STRUCT_FIELD(FSessionSettingsUpdate, SchemaName),
-	ONLINE_STRUCT_FIELD(FSessionSettingsUpdate, NumMaxPublicConnections),
-	ONLINE_STRUCT_FIELD(FSessionSettingsUpdate, NumOpenPublicConnections),
-	ONLINE_STRUCT_FIELD(FSessionSettingsUpdate, NumMaxPrivateConnections),
-	ONLINE_STRUCT_FIELD(FSessionSettingsUpdate, NumOpenPrivateConnections),
+	ONLINE_STRUCT_FIELD(FSessionSettingsUpdate, NumMaxConnections),
 	ONLINE_STRUCT_FIELD(FSessionSettingsUpdate, JoinPolicy),
-	ONLINE_STRUCT_FIELD(FSessionSettingsUpdate, SessionIdOverride),
-	ONLINE_STRUCT_FIELD(FSessionSettingsUpdate, bIsDedicatedServerSession),
 	ONLINE_STRUCT_FIELD(FSessionSettingsUpdate, bAllowNewMembers),
-	ONLINE_STRUCT_FIELD(FSessionSettingsUpdate, bAllowSanctionedPlayers),
-	ONLINE_STRUCT_FIELD(FSessionSettingsUpdate, bAntiCheatProtected),
-	ONLINE_STRUCT_FIELD(FSessionSettingsUpdate, bPresenceEnabled),
 	ONLINE_STRUCT_FIELD(FSessionSettingsUpdate, UpdatedCustomSettings),
 	ONLINE_STRUCT_FIELD(FSessionSettingsUpdate, RemovedCustomSettings),
 	ONLINE_STRUCT_FIELD(FSessionSettingsUpdate, UpdatedSessionMembers),
@@ -875,6 +960,15 @@ BEGIN_ONLINE_STRUCT_META(FGetPresenceSession::Result)
 	ONLINE_STRUCT_FIELD(FGetPresenceSession::Result, Session)
 END_ONLINE_STRUCT_META()
 
+BEGIN_ONLINE_STRUCT_META(FIsPresenceSession::Params)
+ONLINE_STRUCT_FIELD(FIsPresenceSession::Params, LocalAccountId),
+ONLINE_STRUCT_FIELD(FIsPresenceSession::Params, SessionId)
+END_ONLINE_STRUCT_META()
+
+BEGIN_ONLINE_STRUCT_META(FIsPresenceSession::Result)
+	ONLINE_STRUCT_FIELD(FIsPresenceSession::Result, bIsPresenceSession)
+END_ONLINE_STRUCT_META()
+
 BEGIN_ONLINE_STRUCT_META(FSetPresenceSession::Params)
 	ONLINE_STRUCT_FIELD(FSetPresenceSession::Params, LocalAccountId),
 	ONLINE_STRUCT_FIELD(FSetPresenceSession::Params, SessionId)
@@ -894,6 +988,12 @@ BEGIN_ONLINE_STRUCT_META(FCreateSession::Params)
 	ONLINE_STRUCT_FIELD(FCreateSession::Params, LocalAccountId),
 	ONLINE_STRUCT_FIELD(FCreateSession::Params, SessionName),
 	ONLINE_STRUCT_FIELD(FCreateSession::Params, SessionMemberData),
+	ONLINE_STRUCT_FIELD(FCreateSession::Params, SessionIdOverride),
+	ONLINE_STRUCT_FIELD(FCreateSession::Params, bPresenceEnabled),
+	ONLINE_STRUCT_FIELD(FCreateSession::Params, bIsLANSession),
+	ONLINE_STRUCT_FIELD(FCreateSession::Params, bIsDedicatedServerSession),
+	ONLINE_STRUCT_FIELD(FCreateSession::Params, bAllowSanctionedPlayers),
+	ONLINE_STRUCT_FIELD(FCreateSession::Params, bAntiCheatProtected),
 	ONLINE_STRUCT_FIELD(FCreateSession::Params, SessionSettings)
 END_ONLINE_STRUCT_META()
 
@@ -932,11 +1032,8 @@ BEGIN_ONLINE_STRUCT_META(FFindSessions::Result)
 END_ONLINE_STRUCT_META()
 
 BEGIN_ONLINE_STRUCT_META(FStartMatchmaking::Params)
-	ONLINE_STRUCT_FIELD(FStartMatchmaking::Params, LocalAccountId),
-	ONLINE_STRUCT_FIELD(FStartMatchmaking::Params, SessionName),
-	ONLINE_STRUCT_FIELD(FStartMatchmaking::Params, SessionMemberData),
-	ONLINE_STRUCT_FIELD(FStartMatchmaking::Params, SessionSettings),
-	ONLINE_STRUCT_FIELD(FStartMatchmaking::Params, SearchFilters)
+	ONLINE_STRUCT_FIELD(FStartMatchmaking::Params, SessionCreationParameters),
+	ONLINE_STRUCT_FIELD(FStartMatchmaking::Params, SessionSearchFilters)
 END_ONLINE_STRUCT_META()
 
 BEGIN_ONLINE_STRUCT_META(FStartMatchmaking::Result)
@@ -946,7 +1043,8 @@ BEGIN_ONLINE_STRUCT_META(FJoinSession::Params)
 	ONLINE_STRUCT_FIELD(FJoinSession::Params, LocalAccountId),
 	ONLINE_STRUCT_FIELD(FJoinSession::Params, SessionName),
 	ONLINE_STRUCT_FIELD(FJoinSession::Params, SessionId),
-	ONLINE_STRUCT_FIELD(FJoinSession::Params, SessionMemberData)
+	ONLINE_STRUCT_FIELD(FJoinSession::Params, SessionMemberData),
+	ONLINE_STRUCT_FIELD(FJoinSession::Params, bPresenceEnabled)
 END_ONLINE_STRUCT_META()
 
 BEGIN_ONLINE_STRUCT_META(FJoinSession::Result)
@@ -994,6 +1092,29 @@ END_ONLINE_STRUCT_META()
 BEGIN_ONLINE_STRUCT_META(FRejectSessionInvite::Result)
 END_ONLINE_STRUCT_META()
 
+BEGIN_ONLINE_STRUCT_META(FSessionSettingsChanges)
+	ONLINE_STRUCT_FIELD(FSessionSettingsChanges, SchemaName),
+	ONLINE_STRUCT_FIELD(FSessionSettingsChanges, NumMaxConnections),
+	ONLINE_STRUCT_FIELD(FSessionSettingsChanges, JoinPolicy),
+	ONLINE_STRUCT_FIELD(FSessionSettingsChanges, bAllowNewMembers),
+	ONLINE_STRUCT_FIELD(FSessionSettingsChanges, AddedCustomSettings),
+	ONLINE_STRUCT_FIELD(FSessionSettingsChanges, ChangedCustomSettings),
+	ONLINE_STRUCT_FIELD(FSessionSettingsChanges, RemovedCustomSettings)
+END_ONLINE_STRUCT_META()
+
+BEGIN_ONLINE_STRUCT_META(FSessionMemberChanges)
+	ONLINE_STRUCT_FIELD(FSessionMemberChanges, AddedMemberSettings),
+	ONLINE_STRUCT_FIELD(FSessionMemberChanges, ChangedMemberSettings),
+	ONLINE_STRUCT_FIELD(FSessionMemberChanges, RemovedMemberSettings)
+END_ONLINE_STRUCT_META()
+
+BEGIN_ONLINE_STRUCT_META(FSessionUpdate)
+	ONLINE_STRUCT_FIELD(FSessionUpdate, OwnerAccountId),
+	ONLINE_STRUCT_FIELD(FSessionUpdate, SessionSettingsChanges),
+	ONLINE_STRUCT_FIELD(FSessionUpdate, SessionMembersChanges),
+	ONLINE_STRUCT_FIELD(FSessionUpdate, RemovedSessionMembers)
+END_ONLINE_STRUCT_META()
+
 /* Meta*/ }
 
 /* UE::Online */ }
\ No newline at end of file
diff --git a/Engine/Plugins/Online/Experimental/OnlineServicesEOSGS/Source/Private/Online/SessionsEOSGS.cpp b/Engine/Plugins/Online/Experimental/OnlineServicesEOSGS/Source/Private/Online/SessionsEOSGS.cpp
--- a/Engine/Plugins/Online/Experimental/OnlineServicesEOSGS/Source/Private/Online/SessionsEOSGS.cpp
+++ b/Engine/Plugins/Online/Experimental/OnlineServicesEOSGS/Source/Private/Online/SessionsEOSGS.cpp
@@ -53,13 +53,14 @@ FSessionEOSGS::FSessionEOSGS(const EOS_HSessionDetails& InSessionDetailsHandle)
 	EOS_EResult CopyInfoResult = EOS_SessionDetails_CopyInfo(InSessionDetailsHandle, &CopyInfoOptions, &SessionDetailsInfo);
 	if (CopyInfoResult == EOS_EResult::EOS_Success)
 	{
-		SessionId = FSessionsEOSGS::CreateSessionId(FString(SessionDetailsInfo->SessionId));
-		SessionSettings.NumOpenPublicConnections = SessionDetailsInfo->NumOpenPublicConnections;
+		SessionInfo.SessionId = FSessionsEOSGS::CreateSessionId(FString(SessionDetailsInfo->SessionId));
+		SessionInfo.bAllowSanctionedPlayers = !SessionDetailsInfo->Settings->bSanctionsEnabled;
+
 		// We could retrieve the Host Address here if we set it during creation or update
 
-		// TODO: evaluate effect of bInvitesAllowed and bAllowJoinInProgress to know whether we can just read the values here or if we need the custom parameter
-		SessionSettings.bAllowSanctionedPlayers = !SessionDetailsInfo->Settings->bSanctionsEnabled;
-		SessionSettings.NumMaxPublicConnections = SessionDetailsInfo->Settings->NumPublicConnections;
+		// bInvitesAllowed and bAllowJoinInProgress should both have the same value, and that value corresponds to bAllowNewMembers
+		SessionSettings.bAllowNewMembers = SessionDetailsInfo->Settings->bInvitesAllowed && SessionDetailsInfo->Settings->bAllowJoinInProgress;
+		SessionSettings.NumMaxConnections = SessionDetailsInfo->Settings->NumPublicConnections;
 		SessionSettings.JoinPolicy = FromServiceType(SessionDetailsInfo->Settings->PermissionLevel);
 		SessionSettings.CustomSettings.Emplace(EOSGS_BUCKET_ID, FCustomSessionSetting{ FSchemaVariant(FString(UTF8_TO_TCHAR(SessionDetailsInfo->Settings->BucketId))), ESchemaAttributeVisibility::Public });
 
@@ -95,7 +96,7 @@ FSessionEOSGS::FSessionEOSGS(const EOS_HSessionDetails& InSessionDetailsHandle)
 					EOS_ProductUserId ProductUserId = EOS_ProductUserId_FromString(TCHAR_TO_UTF8(*PlayerIdStr));
 					FAccountId AccountId = FindAccountId(ProductUserId);
 
-					FSessionMember& SessionMember = SessionSettings.SessionMembers.FindOrAdd(AccountId);
+					FSessionMember& SessionMember = SessionMembers.FindOrAdd(AccountId);
 
 					// And add the corresponding custom setting
 					FSchemaAttributeId AttributeId(KeyComponents[1]);
@@ -121,15 +122,11 @@ FSessionEOSGS::FSessionEOSGS(const EOS_HSessionDetails& InSessionDetailsHandle)
 					}
 					else if (Key == EOSGS_ANTI_CHEAT_PROTECTED.ToString())
 					{
-						SessionSettings.bAntiCheatProtected = CustomSettingData.Value.GetBoolean();
+						SessionInfo.bAntiCheatProtected = CustomSettingData.Value.GetBoolean();
 					}
 					else if (Key == EOSGS_IS_DEDICATED_SERVER_SESSION.ToString())
 					{
-						SessionSettings.bIsDedicatedServerSession = CustomSettingData.Value.GetBoolean();
-					}
-					else if (Key == EOSGS_PRESENCE_ENABLED.ToString())
-					{
-						SessionSettings.bPresenceEnabled = CustomSettingData.Value.GetBoolean();
+						SessionInfo.bIsDedicatedServerSession = CustomSettingData.Value.GetBoolean();
 					}
 					else if (Key == EOSGS_SCHEMA_NAME.ToString())
 					{
@@ -137,7 +134,7 @@ FSessionEOSGS::FSessionEOSGS(const EOS_HSessionDetails& InSessionDetailsHandle)
 					}
 					else if (Key == EOSGS_SESSION_ID_OVERRIDE.ToString())
 					{
-						SessionSettings.SessionIdOverride = CustomSettingData.Value.GetString();
+						SessionInfo.SessionIdOverride = CustomSettingData.Value.GetString();
 					}
 					else // The rest are parsed as a Custom Session Setting
 					{
@@ -261,7 +258,7 @@ void FSessionsEOSGS::HandleSessionInviteReceived(const EOS_Sessions_SessionInvit
 						ReceiverId,
 						SenderId,
 						SessionInviteId,
-						Session->SessionId
+						Session->GetSessionId()
 						});
 
 					AddSessionInvite(SessionInviteRef, Session, ReceiverId);
@@ -292,11 +289,11 @@ void FSessionsEOSGS::HandleSessionInviteAccepted(const EOS_Sessions_SessionInvit
 				if (const TSharedPtr<ISessions> StrongThis = WeakThis.Pin())
 				{
 					// Instead of using the session information we have stored, we'll use this opportunity to get updated data for the session
-					AllSessionsById.Emplace(Result.GetOkValue().Session->SessionId, Result.GetOkValue().Session);
+					AllSessionsById.Emplace(Result.GetOkValue().Session->GetSessionId(), Result.GetOkValue().Session);
 
 					FUISessionJoinRequested Event{
 									Result.GetOkValue().LocalAccountId,
-									TResult<FOnlineSessionId, FOnlineError>(Result.GetOkValue().Session->SessionId),
+									TResult<FOnlineSessionId, FOnlineError>(Result.GetOkValue().Session->GetSessionId()),
 						EUISessionJoinRequestedSource::FromInvitation
 					};
 
@@ -322,11 +319,11 @@ void FSessionsEOSGS::HandleJoinSessionAccepted(const EOS_Sessions_JoinSessionAcc
 				if (const TSharedPtr<ISessions> StrongThis = WeakThis.Pin())
 				{
 					// Instead of using the session information we have stored, we'll use this opportunity to get updated data for the session
-					AllSessionsById.Emplace(Result.GetOkValue().Session->SessionId, Result.GetOkValue().Session);
+					AllSessionsById.Emplace(Result.GetOkValue().Session->GetSessionId(), Result.GetOkValue().Session);
 
 					FUISessionJoinRequested Event{
 						Result.GetOkValue().LocalAccountId,
-						TResult<FOnlineSessionId, FOnlineError>(Result.GetOkValue().Session->SessionId),
+						TResult<FOnlineSessionId, FOnlineError>(Result.GetOkValue().Session->GetSessionId()),
 			EUISessionJoinRequestedSource::Unspecified
 					};
 
@@ -342,7 +339,7 @@ void FSessionsEOSGS::HandleJoinSessionAccepted(const EOS_Sessions_JoinSessionAcc
 TOnlineAsyncOpHandle<FCreateSession> FSessionsEOSGS::CreateSession(FCreateSession::Params&& Params)
 {
 	// LAN Sessions
-	if (Params.SessionSettings.bIsLANSession)
+	if (Params.bIsLANSession)
 	{
 		return FSessionsLAN::CreateSession(MoveTemp(Params));
 	}
@@ -371,13 +368,13 @@ TOnlineAsyncOpHandle<FCreateSession> FSessionsEOSGS::CreateSession(FCreateSessio
 			return MakeFulfilledPromise<TDefaultErrorResult<FUpdateSessionImpl>>().GetFuture();
 		}
 
-		if (!OpParams.SessionSettings.SessionIdOverride.IsEmpty())
+		if (!OpParams.SessionIdOverride.IsEmpty())
 		{
-			int32 Length = OpParams.SessionSettings.SessionIdOverride.Len();
+			int32 Length = OpParams.SessionIdOverride.Len();
 
 			if (Length < EOS_SESSIONMODIFICATION_MIN_SESSIONIDOVERRIDE_LENGTH || Length > EOS_SESSIONMODIFICATION_MAX_SESSIONIDOVERRIDE_LENGTH)
 			{
-				UE_LOG(LogTemp, Warning, TEXT("[FSessionsEOSGS::CreateSession] Could not create session with SessionIdOverride [%s] of size [%d]. SessionIdOverride size must be between [%d] and [%d] characters long"), *OpParams.SessionSettings.SessionIdOverride, Length, EOS_SESSIONMODIFICATION_MIN_SESSIONIDOVERRIDE_LENGTH, EOS_SESSIONMODIFICATION_MAX_SESSIONIDOVERRIDE_LENGTH);
+				UE_LOG(LogTemp, Warning, TEXT("[FSessionsEOSGS::CreateSession] Could not create session with SessionIdOverride [%s] of size [%d]. SessionIdOverride size must be between [%d] and [%d] characters long"), *OpParams.SessionIdOverride, Length, EOS_SESSIONMODIFICATION_MIN_SESSIONIDOVERRIDE_LENGTH, EOS_SESSIONMODIFICATION_MAX_SESSIONIDOVERRIDE_LENGTH);
 
 				Op.SetError(Errors::InvalidParams());
 
@@ -401,8 +398,8 @@ TOnlineAsyncOpHandle<FCreateSession> FSessionsEOSGS::CreateSession(FCreateSessio
 		CreateSessionModificationOptions.ApiVersion = EOS_SESSIONS_CREATESESSIONMODIFICATION_API_LATEST;
 		static_assert(EOS_SESSIONS_CREATESESSIONMODIFICATION_API_LATEST == 4, "EOS_Sessions_CreateSessionModificationOptions updated, check new fields");
 
-		CreateSessionModificationOptions.bPresenceEnabled = OpParams.SessionSettings.bPresenceEnabled;
-		CreateSessionModificationOptions.bSanctionsEnabled = !OpParams.SessionSettings.bAllowSanctionedPlayers;
+		CreateSessionModificationOptions.bPresenceEnabled = OpParams.bPresenceEnabled;
+		CreateSessionModificationOptions.bSanctionsEnabled = !OpParams.bAllowSanctionedPlayers;
 
 		const FCustomSessionSetting* BucketIdSetting = OpParams.SessionSettings.CustomSettings.Find(EOSGS_BUCKET_ID);
 		const FTCHARToUTF8 BucketIdUtf8(BucketIdSetting ? *BucketIdSetting->Data.GetString() : nullptr);
@@ -413,9 +410,9 @@ TOnlineAsyncOpHandle<FCreateSession> FSessionsEOSGS::CreateSession(FCreateSessio
 		}
 	
 		CreateSessionModificationOptions.LocalUserId = GetProductUserIdChecked(OpParams.LocalAccountId);
-		CreateSessionModificationOptions.MaxPlayers = OpParams.SessionSettings.NumMaxPrivateConnections + OpParams.SessionSettings.NumMaxPublicConnections;
+		CreateSessionModificationOptions.MaxPlayers = OpParams.SessionSettings.NumMaxConnections;
 
-		const FTCHARToUTF8 SessionIdUtf8(*OpParams.SessionSettings.SessionIdOverride);
+		const FTCHARToUTF8 SessionIdUtf8(*OpParams.SessionIdOverride);
 		if (SessionIdUtf8.Length())
 		{
 			CreateSessionModificationOptions.SessionId = SessionIdUtf8.Get();
@@ -438,7 +435,7 @@ TOnlineAsyncOpHandle<FCreateSession> FSessionsEOSGS::CreateSession(FCreateSessio
 		// TODO: We could call EOS_SessionModification_SetHostAddress at this point, although it's not necessary
 
 		// We write all SessionSettings values into the SessionModificationHandle
-		WriteCreateSessionModificationHandle(SessionModificationHandle, OpParams.SessionSettings);
+		WriteCreateSessionModificationHandle(SessionModificationHandle, OpParams);
 
 		// Always update joinability on session creation
 		FUpdateSessionImpl::Params UpdateSessionImplParams{ MakeShared<FSessionModificationHandleEOSGS>(SessionModificationHandle), FUpdateSessionJoinabilityParams{ OpParams.SessionName, OpParams.SessionSettings.bAllowNewMembers } };
@@ -455,9 +452,9 @@ TOnlineAsyncOpHandle<FCreateSession> FSessionsEOSGS::CreateSession(FCreateSessio
 			TSharedRef<FSessionEOSGS> NewSessionEOSGSRef = MakeShared<FSessionEOSGS>();
 			NewSessionEOSGSRef->OwnerAccountId = OpParams.LocalAccountId;
 			NewSessionEOSGSRef->SessionSettings = OpParams.SessionSettings;
-			NewSessionEOSGSRef->SessionId = CreateSessionId(Result.GetOkValue().NewSessionId);
+			NewSessionEOSGSRef->SessionInfo.SessionId = CreateSessionId(Result.GetOkValue().NewSessionId);
 
-			AddSessionWithReferences(NewSessionEOSGSRef, OpParams.SessionName, OpParams.LocalAccountId, NewSessionEOSGSRef->SessionSettings.bPresenceEnabled);
+			AddSessionWithReferences(NewSessionEOSGSRef, OpParams.SessionName, OpParams.LocalAccountId, OpParams.bPresenceEnabled);
 
 			Op.SetResult({ });
 		}
@@ -581,61 +578,53 @@ void FSessionsEOSGS::RemoveAttribute(EOS_HSessionModification& SessionModificati
 	}
 }
 
-void FSessionsEOSGS::WriteCreateSessionModificationHandle(EOS_HSessionModification& SessionModificationHandle, const FSessionSettings& SessionSettings)
+void FSessionsEOSGS::WriteCreateSessionModificationHandle(EOS_HSessionModification& SessionModificationHandle, const FCreateSession::Params& Params)
 {
 	// TODO: We have the option to call EOS_SessionModification_SetHostAddress in EOS, useful if the session owner changes
 
 	// We'll update this setting in the session modification step, and start or end the session accordingly to bock or allow join processes
-	SetJoinInProgressAllowed(SessionModificationHandle, SessionSettings.bAllowNewMembers);
+	SetJoinInProgressAllowed(SessionModificationHandle, Params.SessionSettings.bAllowNewMembers);
 
 	// We'll also update invite permissions for the session
-	SetInvitesAllowed(SessionModificationHandle, SessionSettings.bAllowNewMembers);
+	SetInvitesAllowed(SessionModificationHandle, Params.SessionSettings.bAllowNewMembers);
 
 	// We won't copy bIsLANSession since it' irrelevant for EOS Sessions
-	AddAttribute(SessionModificationHandle, EOSGS_ALLOW_NEW_MEMBERS, { FSchemaVariant(SessionSettings.bAllowNewMembers), ESchemaAttributeVisibility::Public });
-	AddAttribute(SessionModificationHandle, EOSGS_ANTI_CHEAT_PROTECTED, { FSchemaVariant(SessionSettings.bAntiCheatProtected), ESchemaAttributeVisibility::Public });
-	AddAttribute(SessionModificationHandle, EOSGS_IS_DEDICATED_SERVER_SESSION, { FSchemaVariant(SessionSettings.bIsDedicatedServerSession), ESchemaAttributeVisibility::Public });
-	AddAttribute(SessionModificationHandle, EOSGS_PRESENCE_ENABLED, { FSchemaVariant(SessionSettings.bPresenceEnabled), ESchemaAttributeVisibility::Public });
+	AddAttribute(SessionModificationHandle, EOSGS_ALLOW_NEW_MEMBERS, { FSchemaVariant(Params.SessionSettings.bAllowNewMembers), ESchemaAttributeVisibility::Public });
+	AddAttribute(SessionModificationHandle, EOSGS_ANTI_CHEAT_PROTECTED, { FSchemaVariant(Params.bAntiCheatProtected), ESchemaAttributeVisibility::Public });
+	AddAttribute(SessionModificationHandle, EOSGS_IS_DEDICATED_SERVER_SESSION, { FSchemaVariant(Params.bIsDedicatedServerSession), ESchemaAttributeVisibility::Public });
 
-	SetPermissionLevel(SessionModificationHandle, SessionSettings.JoinPolicy);
+	SetPermissionLevel(SessionModificationHandle, Params.SessionSettings.JoinPolicy);
 
-	SetMaxPlayers(SessionModificationHandle, SessionSettings.NumMaxPrivateConnections + SessionSettings.NumMaxPublicConnections);
+	SetMaxPlayers(SessionModificationHandle, Params.SessionSettings.NumMaxConnections);
 
-	AddAttribute(SessionModificationHandle, EOSGS_SCHEMA_NAME, { FSchemaVariant(SessionSettings.SchemaName.ToString()), ESchemaAttributeVisibility::Public });
-	AddAttribute(SessionModificationHandle, EOSGS_SESSION_ID_OVERRIDE, { FSchemaVariant(SessionSettings.SessionIdOverride), ESchemaAttributeVisibility::Public });
+	AddAttribute(SessionModificationHandle, EOSGS_SCHEMA_NAME, { FSchemaVariant(Params.SessionSettings.SchemaName.ToString()), ESchemaAttributeVisibility::Public });
+	AddAttribute(SessionModificationHandle, EOSGS_SESSION_ID_OVERRIDE, { FSchemaVariant(Params.SessionIdOverride), ESchemaAttributeVisibility::Public });
 
 	// Custom Settings
 
- 	for (const TPair<FSchemaAttributeId, FCustomSessionSetting>& Entry : SessionSettings.CustomSettings)
+ 	for (const TPair<FSchemaAttributeId, FCustomSessionSetting>& Entry : Params.SessionSettings.CustomSettings)
  	{
 		AddAttribute(SessionModificationHandle, Entry.Key, Entry.Value);
 	}
 
 	// BucketId has its own set method on the API
-	const FCustomSessionSetting* NewBucketIdSetting = SessionSettings.CustomSettings.Find(EOSGS_BUCKET_ID);
+	const FCustomSessionSetting* NewBucketIdSetting = Params.SessionSettings.CustomSettings.Find(EOSGS_BUCKET_ID);
 	if (NewBucketIdSetting)
 	{
 		SetBucketId(SessionModificationHandle, NewBucketIdSetting->Data.GetString());
 	}
 
-	// Session Members
-
-	for (const TPair<FAccountId, FSessionMember>& SessionMemberEntry : SessionSettings.SessionMembers)
+	// Session Member	
+	for (const TPair<FSchemaAttributeId, FCustomSessionSetting>& CustomSettingEntry : Params.SessionMemberData.MemberSettings)
 	{
-		const FSessionMember& SessionMember = SessionMemberEntry.Value;
-		
-		for (const TPair<FSchemaAttributeId, FCustomSessionSetting>& CustomSettingEntry : SessionMember.MemberSettings)
-		{
-			const FName& Key = FName(LexToString(GetProductUserIdChecked(SessionMemberEntry.Key)) + TEXT(":") + CustomSettingEntry.Key.ToString());
-			AddAttribute(SessionModificationHandle, Key, CustomSettingEntry.Value);
-		}		
-	}
+		const FName& Key = FName(LexToString(GetProductUserIdChecked(Params.LocalAccountId)) + TEXT(":") + CustomSettingEntry.Key.ToString());
+		AddAttribute(SessionModificationHandle, Key, CustomSettingEntry.Value);
+	}		
 }
 
 void FSessionsEOSGS::WriteUpdateSessionModificationHandle(EOS_HSessionModification& SessionModificationHandle, const FName& SessionName, const FSessionSettingsUpdate& NewSettings)
 {
 	// TODO: We have the option to call EOS_SessionModification_SetHostAddress in EOS, useful if the session owner changes
-	
 
 	if (NewSettings.bAllowNewMembers.IsSet())
 	{
@@ -646,31 +635,11 @@ void FSessionsEOSGS::WriteUpdateSessionModificationHandle(EOS_HSessionModificati
 		SetInvitesAllowed(SessionModificationHandle, NewSettings.bAllowNewMembers.GetValue());
 	}
 
-	if (NewSettings.bAntiCheatProtected.IsSet())
-	{
-		AddAttribute(SessionModificationHandle, EOSGS_ANTI_CHEAT_PROTECTED, { FSchemaVariant(NewSettings.bAntiCheatProtected.GetValue()), ESchemaAttributeVisibility::Public });
-	}
-
-	if (NewSettings.bIsDedicatedServerSession.IsSet())
-	{
-		AddAttribute(SessionModificationHandle, EOSGS_IS_DEDICATED_SERVER_SESSION, { FSchemaVariant(NewSettings.bIsDedicatedServerSession.GetValue()), ESchemaAttributeVisibility::Public });
-	}
-
-	if (NewSettings.bPresenceEnabled.IsSet())
-	{
-		AddAttribute(SessionModificationHandle, EOSGS_PRESENCE_ENABLED, { FSchemaVariant(NewSettings.bPresenceEnabled.GetValue()), ESchemaAttributeVisibility::Public });
-	}
-
 	if (NewSettings.SchemaName.IsSet())
 	{
 		AddAttribute(SessionModificationHandle, EOSGS_SCHEMA_NAME, { FSchemaVariant(NewSettings.SchemaName.GetValue().ToString()), ESchemaAttributeVisibility::Public });
 	}
 
-	if (NewSettings.SessionIdOverride.IsSet())
-	{
-		AddAttribute(SessionModificationHandle, EOSGS_SESSION_ID_OVERRIDE, { FSchemaVariant(NewSettings.SessionIdOverride.GetValue()), ESchemaAttributeVisibility::Public });
-	}
-
 	TOnlineResult<FGetSessionByName> GetSessionByNameResult = GetSessionByName({ SessionName });
 	check(GetSessionByNameResult.IsOk());
 
@@ -682,13 +651,9 @@ void FSessionsEOSGS::WriteUpdateSessionModificationHandle(EOS_HSessionModificati
 		SetPermissionLevel(SessionModificationHandle, NewSettings.JoinPolicy.GetValue());
 	}
 
-	if (NewSettings.NumMaxPrivateConnections.IsSet() || NewSettings.NumMaxPublicConnections.IsSet())
+	if (NewSettings.NumMaxConnections.IsSet())
 	{
-		// We set MaxPlayers to the sum of any new values (if set) and old values (if not set)
-		uint32 DefaultNumMaxPrivateConnections = SessionSettings.NumMaxPrivateConnections;
-		uint32 DefaultNumMaxPublicConnections = SessionSettings.NumMaxPublicConnections;
-
-		SetMaxPlayers(SessionModificationHandle, NewSettings.NumMaxPrivateConnections.Get(DefaultNumMaxPrivateConnections) + NewSettings.NumMaxPublicConnections.Get(DefaultNumMaxPublicConnections));
+		SetMaxPlayers(SessionModificationHandle, NewSettings.NumMaxConnections.GetValue());
 	}
 
 	// BucketId has its own set method on the API
@@ -715,7 +680,7 @@ void FSessionsEOSGS::WriteUpdateSessionModificationHandle(EOS_HSessionModificati
 	for (const FAccountId& SessionMemberId : NewSettings.RemovedSessionMembers)
 	{
 		// To remove a session member, we'll remove all their attributes
- 		if (const FSessionMember* SessionMember = SessionSettings.SessionMembers.Find(SessionMemberId))
+ 		if (const FSessionMember* SessionMember = FoundSession->GetSessionMembers().Find(SessionMemberId))
  		{
  			TArray<FName> SettingsKeyArray;
  			SessionMember->MemberSettings.GenerateKeyArray(SettingsKeyArray);
@@ -749,7 +714,7 @@ TOnlineAsyncOpHandle<FUpdateSession> FSessionsEOSGS::UpdateSession(FUpdateSessio
 {
 	// LAN Sessions
 	TOnlineResult<FGetSessionByName> Result = GetSessionByName({ Params.SessionName });
-	if (Result.IsOk() && Result.GetOkValue().Session->GetSessionSettings().bIsLANSession)
+	if (Result.IsOk() && Result.GetOkValue().Session->GetSessionInfo().bIsLANSession)
 	{
 		return FSessionsLAN::UpdateSession(MoveTemp(Params));
 	}
@@ -815,13 +780,15 @@ TOnlineAsyncOpHandle<FUpdateSession> FSessionsEOSGS::UpdateSession(FUpdateSessio
 				TSharedRef<FSessionCommon> FoundSession = GetMutableSessionByNameResult.GetOkValue().Session;
 
 				// Now that the API Session update has processed successfully, we'll update our local session with the same data
-				FoundSession->SessionSettings += OpParams.Mutations;
+				//FoundSession->SessionSettings += OpParams.Mutations;
 
 				// We set the result and fire the event
 				Op.SetResult({ });
 
-				FSessionUpdated SessionUpdatedEvent{ OpParams.SessionName, OpParams.Mutations };
-				SessionEvents.OnSessionUpdated.Broadcast(SessionUpdatedEvent);
+				//FSessionUpdated SessionUpdatedEvent{ OpParams.SessionName, OpParams.Mutations };
+				//SessionEvents.OnSessionUpdated.Broadcast(SessionUpdatedEvent);
+
+				// TODO: Will be refactored in the UpdateSession split
 			}
 			else
 			{
@@ -1019,7 +986,7 @@ TOnlineAsyncOpHandle<FLeaveSession> FSessionsEOSGS::LeaveSession(FLeaveSession::
 {
 	// LAN Sessions
 	TOnlineResult<FGetSessionByName> Result = GetSessionByName({ Params.SessionName });
-	if (Result.IsOk() && Result.GetOkValue().Session->GetSessionSettings().bIsLANSession)
+	if (Result.IsOk() && Result.GetOkValue().Session->GetSessionInfo().bIsLANSession)
 	{
 		return FSessionsLAN::LeaveSession(MoveTemp(Params));
 	}
@@ -1064,7 +1031,7 @@ TOnlineAsyncOpHandle<FLeaveSession> FSessionsEOSGS::LeaveSession(FLeaveSession::
 		{
 			TSharedRef<const ISession> FoundSession = GetSessionByNameResult.GetOkValue().Session;
 
-			ClearSessionReferences(FoundSession->GetSessionId(), OpParams.SessionName, OpParams.LocalAccountId, FoundSession->GetSessionSettings().bPresenceEnabled);
+			ClearSessionReferences(FoundSession->GetSessionId(), OpParams.SessionName, OpParams.LocalAccountId);
 		}
 
 		Op.SetResult(FLeaveSession::Result{ });
@@ -1358,7 +1325,7 @@ TOnlineAsyncOpHandle<FJoinSession> FSessionsEOSGS::JoinSession(FJoinSession::Par
 	const TSharedRef<const ISession>& FoundSession = GetSessionByIdResult.GetOkValue().Session;
 
 	// LAN Sessions
-	if (FoundSession->GetSessionSettings().bIsLANSession)
+	if (FoundSession->GetSessionInfo().bIsLANSession)
 	{
 		return FSessionsLAN::JoinSession(MoveTemp(Params));
 	}
@@ -1416,7 +1383,7 @@ TOnlineAsyncOpHandle<FJoinSession> FSessionsEOSGS::JoinSession(FJoinSession::Par
 		JoinSessionOptions.ApiVersion = EOS_SESSIONS_JOINSESSION_API_LATEST;
 		static_assert(EOS_SESSIONS_JOINSESSION_API_LATEST == 2, "EOS_Sessions_JoinSessionOptions updated, check new fields");
 
-		JoinSessionOptions.bPresenceEnabled = FoundSession->GetSessionSettings().bPresenceEnabled;
+		JoinSessionOptions.bPresenceEnabled = OpParams.bPresenceEnabled;
 
 		JoinSessionOptions.LocalUserId = GetProductUserIdChecked(OpParams.LocalAccountId);
 
@@ -1453,7 +1420,7 @@ TOnlineAsyncOpHandle<FJoinSession> FSessionsEOSGS::JoinSession(FJoinSession::Par
 
 		const TSharedRef<const ISession>& FoundSession = GetSessionByIdResult.GetOkValue().Session;
 
-		AddSessionReferences(FoundSession->GetSessionId(), OpParams.SessionName, OpParams.LocalAccountId, FoundSession->GetSessionSettings().bPresenceEnabled);
+		AddSessionReferences(FoundSession->GetSessionId(), OpParams.SessionName, OpParams.LocalAccountId, OpParams.bPresenceEnabled);
 
 		// After successfully joining a session, we'll remove all related invites if any are found
 		ClearSessionInvitesForSession(OpParams.LocalAccountId, FoundSession->GetSessionId());
@@ -1727,7 +1694,7 @@ TOnlineAsyncOpHandle<FAddSessionMember> FSessionsEOSGS::AddSessionMember(FAddSes
 {
 	// LAN Sessions
 	TOnlineResult<FGetSessionByName> Result = GetSessionByName({ Params.SessionName });
-	if (Result.IsOk() && Result.GetOkValue().Session->GetSessionSettings().bIsLANSession)
+	if (Result.IsOk() && Result.GetOkValue().Session->GetSessionInfo().bIsLANSession)
 	{
 		return FSessionsLAN::AddSessionMember(MoveTemp(Params));
 	}
@@ -1790,7 +1757,7 @@ TOnlineAsyncOpHandle<FRemoveSessionMember> FSessionsEOSGS::RemoveSessionMember(F
 {
 	// LAN Sessions
 	TOnlineResult<FGetSessionByName> Result = GetSessionByName({ Params.SessionName });
-	if (Result.IsOk() && Result.GetOkValue().Session->GetSessionSettings().bIsLANSession)
+	if (Result.IsOk() && Result.GetOkValue().Session->GetSessionInfo().bIsLANSession)
 	{
 		return FSessionsLAN::RemoveSessionMember(MoveTemp(Params));
 	}
@@ -1858,8 +1825,8 @@ void FSessionsEOSGS::AppendSessionToPacket(FNboSerializeToBuffer& Packet, const
 
 	SerializeToBuffer(Packet, Session);
 	SerializeToBuffer(Packet, Session.OwnerAccountId);
-	SerializeToBuffer(Packet, Session.SessionId);
-	SerializeToBuffer(Packet, Session.SessionSettings.SessionMembers);
+	SerializeToBuffer(Packet, Session.SessionInfo.SessionId);
+	SerializeToBuffer(Packet, Session.SessionMembers);
 }
 
 void FSessionsEOSGS::ReadSessionFromPacket(FNboSerializeFromBuffer& Packet, FSessionLAN& Session)
@@ -1869,8 +1836,8 @@ void FSessionsEOSGS::ReadSessionFromPacket(FNboSerializeFromBuffer& Packet, FSes
 
 	SerializeFromBuffer(Packet, Session);
 	SerializeFromBuffer(Packet, Session.OwnerAccountId);
-	SerializeFromBuffer(Packet, Session.SessionId);
-	SerializeFromBuffer(Packet, Session.SessionSettings.SessionMembers);
+	SerializeFromBuffer(Packet, Session.SessionInfo.SessionId);
+	SerializeFromBuffer(Packet, Session.SessionMembers);
 }
 
 /* UE::Online */ }
\ No newline at end of file
diff --git a/Engine/Plugins/Online/Experimental/OnlineServicesEOSGS/Source/Public/Online/SessionsEOSGS.h b/Engine/Plugins/Online/Experimental/OnlineServicesEOSGS/Source/Public/Online/SessionsEOSGS.h
--- a/Engine/Plugins/Online/Experimental/OnlineServicesEOSGS/Source/Public/Online/SessionsEOSGS.h
+++ b/Engine/Plugins/Online/Experimental/OnlineServicesEOSGS/Source/Public/Online/SessionsEOSGS.h
@@ -223,7 +223,7 @@ protected:
 	/**
 	 * Writes all values in the passed SessionSettings to the SessionModificationHandle
 	 */
-	void WriteCreateSessionModificationHandle(EOS_HSessionModification& SessionModificationHandle, const FSessionSettings& SessionSettings);
+	void WriteCreateSessionModificationHandle(EOS_HSessionModification& SessionModificationHandle, const FCreateSession::Params& Params);
 
 	/**
 	 * Writes only the new values for all updated settings to the SessionModificationHandle
diff --git a/Engine/Plugins/Online/Experimental/OnlineServicesEOSGS/Source/Public/Online/SessionsEOSGSTypes.h b/Engine/Plugins/Online/Experimental/OnlineServicesEOSGS/Source/Public/Online/SessionsEOSGSTypes.h
--- a/Engine/Plugins/Online/Experimental/OnlineServicesEOSGS/Source/Public/Online/SessionsEOSGSTypes.h
+++ b/Engine/Plugins/Online/Experimental/OnlineServicesEOSGS/Source/Public/Online/SessionsEOSGSTypes.h
@@ -15,7 +15,6 @@ namespace UE::Online {
 static FName EOSGS_ALLOW_NEW_MEMBERS = TEXT("EOSGS_ALLOW_NEW_MEMBERS");
 static FName EOSGS_ANTI_CHEAT_PROTECTED = TEXT("EOSGS_ANTI_CHEAT_PROTECTED");
 static FName EOSGS_IS_DEDICATED_SERVER_SESSION = TEXT("EOSGS_IS_DEDICATED_SERVER_SESSION");
-static FName EOSGS_PRESENCE_ENABLED = TEXT("EOSGS_PRESENCE_ENABLED");
 static FName EOSGS_SCHEMA_NAME = TEXT("EOSGS_SCHEMA_NAME");
 static FName EOSGS_SESSION_ID_OVERRIDE = TEXT("EOSGS_SESSION_ID_OVERRIDE");
 
diff --git a/Engine/Plugins/Online/Experimental/OnlineServicesNull/Source/Private/Online/SessionsNull.cpp b/Engine/Plugins/Online/Experimental/OnlineServicesNull/Source/Private/Online/SessionsNull.cpp
--- a/Engine/Plugins/Online/Experimental/OnlineServicesNull/Source/Private/Online/SessionsNull.cpp
+++ b/Engine/Plugins/Online/Experimental/OnlineServicesNull/Source/Private/Online/SessionsNull.cpp
@@ -37,8 +37,8 @@ void FSessionsNull::AppendSessionToPacket(FNboSerializeToBuffer& Packet, const F
 
 	SerializeToBuffer(Packet, Session);
 	SerializeToBuffer(Packet, Session.OwnerAccountId);
-	SerializeToBuffer(Packet, Session.SessionId);
-	SerializeToBuffer(Packet, Session.SessionSettings.SessionMembers);
+	SerializeToBuffer(Packet, Session.SessionInfo.SessionId);
+	SerializeToBuffer(Packet, Session.SessionMembers);
 }
 
 void FSessionsNull::ReadSessionFromPacket(FNboSerializeFromBuffer& Packet, FSessionLAN& Session)
@@ -48,8 +48,8 @@ void FSessionsNull::ReadSessionFromPacket(FNboSerializeFromBuffer& Packet, FSess
 
 	SerializeFromBuffer(Packet, Session);
 	SerializeFromBuffer(Packet, Session.OwnerAccountId);
-	SerializeFromBuffer(Packet, Session.SessionId);
-	SerializeFromBuffer(Packet, Session.SessionSettings.SessionMembers);
+	SerializeFromBuffer(Packet, Session.SessionInfo.SessionId);
+	SerializeFromBuffer(Packet, Session.SessionMembers);
 }
 
 /* UE::Online */ }
\ No newline at end of file
diff --git a/Engine/Plugins/Online/Experimental/OnlineServicesOSSAdapter/Source/Private/Online/SessionsOSSAdapter.cpp b/Engine/Plugins/Online/Experimental/OnlineServicesOSSAdapter/Source/Private/Online/SessionsOSSAdapter.cpp
--- a/Engine/Plugins/Online/Experimental/OnlineServicesOSSAdapter/Source/Private/Online/SessionsOSSAdapter.cpp
+++ b/Engine/Plugins/Online/Experimental/OnlineServicesOSSAdapter/Source/Private/Online/SessionsOSSAdapter.cpp
@@ -237,10 +237,24 @@ TSharedRef<FOnlineSessionSearch> BuildV1SessionSearch(const FFindSessions::Param
 
 /** FSessionOSSAdapter */
 
-FSessionOSSAdapter::FSessionOSSAdapter(const FOnlineSession& InSession)
-	: V1Session(InSession)
+void WriteV2SessionInfoFromV1Session(const FOnlineSession& InSession, FSessionInfo& OutInfo)
 {
+	OutInfo.SessionIdOverride = InSession.SessionSettings.SessionIdOverride;
+
+	InSession.SessionSettings.Settings.FindChecked(OSS_ADAPTER_SESSIONS_ALLOW_SANCTIONED_PLAYERS).Data.GetValue(OutInfo.bAllowSanctionedPlayers);
+
+	OutInfo.bAntiCheatProtected = InSession.SessionSettings.bAntiCheatProtected;
+
+	OutInfo.bIsDedicatedServerSession = InSession.SessionSettings.bIsDedicated;
 
+	OutInfo.bIsLANSession = InSession.SessionSettings.bIsLANMatch;
+}
+
+FSessionOSSAdapter::FSessionOSSAdapter(const FOnlineSession& InSession, const FOnlineSessionId& InSessionId)
+	: V1Session(InSession)
+{
+	WriteV2SessionInfoFromV1Session(InSession, SessionInfo);
+	SessionInfo.SessionId = InSessionId;
 }
 
 const FSessionOSSAdapter& FSessionOSSAdapter::Cast(const ISession& InSession)
@@ -284,7 +298,7 @@ void FSessionsOSSAdapter::Initialize()
 		SessionInvite->SenderId = ServicesOSSAdapter.GetAccountIdRegistry().FindOrAddHandle(FromId.AsShared());
 		SessionInvite->RecipientId = ServicesOSSAdapter.GetAccountIdRegistry().FindOrAddHandle(UserId.AsShared());
 		SessionInvite->InviteId = GetSessionInviteIdRegistry().BasicRegistry.FindOrAddHandle(AppId);
-		SessionInvite->SessionId = Session->SessionId;
+		SessionInvite->SessionId = Session->GetSessionId();
 
 		AddSessionInvite(SessionInvite, Session, SessionInvite->RecipientId);
 
@@ -315,20 +329,24 @@ void FSessionsOSSAdapter::Initialize()
 
 	SessionsInterface->AddOnSessionParticipantRemovedDelegate_Handle(FOnSessionParticipantRemovedDelegate::CreateLambda([this](FName SessionName, const FUniqueNetId& TargetUniqueNetId)
 	{
-			const FOnlineServicesOSSAdapter& ServicesOSSAdapter = static_cast<FOnlineServicesOSSAdapter&>(Services);
-			const FUniqueNetIdRef TargetUniqueNetIdRef = TargetUniqueNetId.AsShared();
-			const FAccountId TargetAccountId = ServicesOSSAdapter.GetAccountIdRegistry().FindOrAddHandle(TargetUniqueNetIdRef);
+		// TODO: Method to update the interface session with the latest OSS Adapter session data
 
-			FSessionSettingsUpdate SettingsUpdate;
-			SettingsUpdate.RemovedSessionMembers.Add(TargetAccountId);
+		const FOnlineServicesOSSAdapter& ServicesOSSAdapter = static_cast<FOnlineServicesOSSAdapter&>(Services);
+		const FUniqueNetIdRef TargetUniqueNetIdRef = TargetUniqueNetId.AsShared();
+		const FAccountId TargetAccountId = ServicesOSSAdapter.GetAccountIdRegistry().FindOrAddHandle(TargetUniqueNetIdRef);
 
-		const FSessionUpdated Event { SessionName , SettingsUpdate };
+		FSessionUpdate SessionUpdate;
+		SessionUpdate.RemovedSessionMembers.Add(TargetAccountId);
 
-			SessionEvents.OnSessionUpdated.Broadcast(Event);
+		const FSessionUpdated Event { SessionName , SessionUpdate };
+
+		SessionEvents.OnSessionUpdated.Broadcast(Event);
 	}));
 
 	SessionsInterface->AddOnSessionParticipantsChangeDelegate_Handle(FOnSessionParticipantsChangeDelegate::CreateLambda([this](FName SessionName, const FUniqueNetId& TargetUniqueNetId, bool bJoined)
 	{
+		// TODO: Method to update the interface session with the latest OSS Adapter session data
+
 		// We won't transmit events for a session or member that doesn't exist
 		const TOnlineResult<FGetSessionByName> GetSessionByNameResult = GetSessionByName({ SessionName });
 		if (GetSessionByNameResult.IsOk())
@@ -339,23 +357,20 @@ void FSessionsOSSAdapter::Initialize()
 			const FUniqueNetIdRef TargetUniqueNetIdRef = TargetUniqueNetId.AsShared();
 			const FAccountId TargetAccountId = ServicesOSSAdapter.GetAccountIdRegistry().FindOrAddHandle(TargetUniqueNetIdRef);
 
-			FSessionSettingsUpdate SettingsUpdate;
+			FSessionUpdate SessionUpdate;
 			if (bJoined)
 			{
-				if (const FSessionMember* SessionMember = FoundSession->GetSessionSettings().SessionMembers.Find(TargetAccountId))
+				if (const FSessionMember* SessionMember = FoundSession->GetSessionMembers().Find(TargetAccountId))
 				{
-					FSessionMemberUpdate SessionMemberUpdate;
-					SessionMemberUpdate.UpdatedMemberSettings.Append(SessionMember->MemberSettings);
-
-					SettingsUpdate.UpdatedSessionMembers.Emplace(TargetAccountId, SessionMemberUpdate);
+					SessionUpdate.AddedSessionMembers.Emplace(TargetAccountId, *SessionMember);
 				}
 			}
 			else
 			{
-				SettingsUpdate.RemovedSessionMembers.Add(TargetAccountId);
+				SessionUpdate.RemovedSessionMembers.Add(TargetAccountId);
 			}
 
-			const FSessionUpdated Event { SessionName , SettingsUpdate };
+			const FSessionUpdated Event { SessionName , SessionUpdate };
 
 			SessionEvents.OnSessionUpdated.Broadcast(Event);
 		}
@@ -373,30 +388,50 @@ void FSessionsOSSAdapter::Initialize()
 			const FUniqueNetIdRef TargetUniqueNetIdRef = TargetUniqueNetId.AsShared();
 			const FAccountId TargetAccountId = ServicesOSSAdapter.GetAccountIdRegistry().FindOrAddHandle(TargetUniqueNetIdRef);
 
-			FSessionSettingsUpdate SettingsUpdate;
-			if (const FSessionMember* SessionMember = FoundSession->GetSessionSettings().SessionMembers.Find(TargetAccountId))
+			FSessionUpdate SessionUpdate;
+			if (const FSessionMember* SessionMember = FoundSession->GetSessionMembers().Find(TargetAccountId))
 			{
 				const ::FSessionSettings* MemberSettings = UpdatedSettings.MemberSettings.Find(TargetUniqueNetIdRef);
 				const FCustomSessionSettingsMap UpdatedMemberSettings = GetV2SessionSettings(*MemberSettings);
 
-				FSessionMemberUpdate SessionMemberUpdate;
+				FSessionMemberChanges SessionMemberChanges;
 				for (const TPair<FName, FCustomSessionSetting>& Entry : UpdatedMemberSettings)
 				{
 					if (const FCustomSessionSetting* Setting = SessionMember->MemberSettings.Find(Entry.Key))
 					{
 						if (Setting->Data != Entry.Value.Data || Setting->Visibility != Entry.Value.Visibility || Setting->ID != Entry.Value.ID)
 						{
-							SessionMemberUpdate.UpdatedMemberSettings.Add(Entry);
+							FCustomSessionSettingUpdate SettingUpdate;
+							SettingUpdate.OldValue = *Setting;
+							SettingUpdate.NewValue = Entry.Value;
+
+							SessionMemberChanges.ChangedMemberSettings.Emplace(Entry.Key, SettingUpdate);
 						}
 					}
+					else
+					{
+						SessionMemberChanges.AddedMemberSettings.Emplace(Entry.Key, Entry.Value);
+					}
+				}
+
+				for (const TPair<FName, FCustomSessionSetting>& Entry : SessionMember->MemberSettings)
+				{
+					if (!UpdatedMemberSettings.Contains(Entry.Key))
+					{
+						SessionMemberChanges.RemovedMemberSettings.AddUnique(Entry.Key);
+					}
 				}
 
-				SettingsUpdate.UpdatedSessionMembers.Emplace(TargetAccountId, SessionMemberUpdate);
+				SessionUpdate.SessionMembersChanges.Emplace(TargetAccountId, SessionMemberChanges);
 			}
 
-			const FSessionUpdated Event{ SessionName , SettingsUpdate };
+			// TODO: Method to update the interface session with the latest OSS Adapter session data
+			// Done after the update to read the value differences
+
+			const FSessionUpdated Event{ SessionName , SessionUpdate };
 
 			SessionEvents.OnSessionUpdated.Broadcast(Event);
+
 		}
 	}));
 }
@@ -455,13 +490,13 @@ TOnlineAsyncOpHandle<FCreateSession> FSessionsOSSAdapter::CreateSession(FCreateS
 
 				TSharedRef<FSessionCommon> V2Session = BuildV2Session(V1Session);
 
-				AddSessionWithReferences(V2Session, OpParams.SessionName, OpParams.LocalAccountId, V2Session->SessionSettings.bPresenceEnabled);
+				AddSessionWithReferences(V2Session, OpParams.SessionName, OpParams.LocalAccountId, OpParams.bPresenceEnabled);
 
 				Op->SetResult({ });
 			}
 		});
 
-		SessionsInterface->CreateSession(*LocalAccountId, OpParams.SessionName, BuildV1Settings(OpParams.SessionSettings));
+		SessionsInterface->CreateSession(*LocalAccountId, OpParams.SessionName, BuildV1SettingsForCreate(OpParams));
 
 		// TODO: Pending support for multiple local users
 	})
@@ -505,13 +540,15 @@ TOnlineAsyncOpHandle<FUpdateSession> FSessionsOSSAdapter::UpdateSession(FUpdateS
 					const TSharedRef<FSessionCommon>& FoundSession = GetMutableSessionByNameResult.GetOkValue().Session;
 
 					FSessionSettings& UpdatedV2Settings = FoundSession->SessionSettings;
-					UpdatedV2Settings += OpParams.Mutations;
+					//UpdatedV2Settings += OpParams.Mutations;
 
 					Op->SetResult({ });
 
-					FSessionUpdated Event { SessionName, OpParams.Mutations };
+					//FSessionUpdated Event { SessionName, OpParams.Mutations };
 
-					SessionEvents.OnSessionUpdated.Broadcast(Event);
+					//SessionEvents.OnSessionUpdated.Broadcast(Event);
+
+					// TODO: Will be refactored as part of the UpdateSession split
 				}
 				else
 				{
@@ -527,9 +564,11 @@ TOnlineAsyncOpHandle<FUpdateSession> FSessionsOSSAdapter::UpdateSession(FUpdateS
 
 		// We will update a copy here, and wait until the operation has completed successfully to update our local data
 		FSessionSettings UpdatedV2Settings = FoundSession->SessionSettings;
-		UpdatedV2Settings += OpParams.Mutations;
+		//UpdatedV2Settings += OpParams.Mutations;
 		
-		FOnlineSessionSettings UpdatedV1Settings = BuildV1Settings(UpdatedV2Settings);
+		// TODO: Will be refactored as part of the UpdateSession split
+
+		FOnlineSessionSettings UpdatedV1Settings = BuildV1SettingsForUpdate(OpParams.LocalAccountId, FoundSession);
 
 		SessionsInterface->UpdateSession(OpParams.SessionName, UpdatedV1Settings);
 	})
@@ -573,7 +612,7 @@ TOnlineAsyncOpHandle<FLeaveSession> FSessionsOSSAdapter::LeaveSession(FLeaveSess
 				{
 					TSharedRef<const ISession> FoundSession = GetSessionByNameResult.GetOkValue().Session;
 
-					ClearSessionReferences(FoundSession->GetSessionId(), OpParams.SessionName, OpParams.LocalAccountId, FoundSession->GetSessionSettings().bPresenceEnabled);
+					ClearSessionReferences(FoundSession->GetSessionId(), OpParams.SessionName, OpParams.LocalAccountId);
 				}
 
 				Op->SetResult({ });
@@ -798,7 +837,7 @@ TOnlineAsyncOpHandle<FStartMatchmaking> FSessionsOSSAdapter::StartMatchmaking(FS
 			return;
 		}
 
-		if (PendingV1SessionSearchesPerUser.Contains(OpParams.LocalAccountId))
+		if (PendingV1SessionSearchesPerUser.Contains(OpParams.SessionCreationParameters.LocalAccountId))
 		{
 			Op.SetError(Errors::AlreadyPending());
 			return;
@@ -808,9 +847,9 @@ TOnlineAsyncOpHandle<FStartMatchmaking> FSessionsOSSAdapter::StartMatchmaking(FS
 
 		TArray<FSessionMatchmakingUser> MatchMakingUsers;
 
-		FSessionMatchmakingUser NewMatchMakingUser { ServicesOSSAdapter.GetAccountIdRegistry().GetIdValue(OpParams.LocalAccountId).ToSharedRef() };
+		FSessionMatchmakingUser NewMatchMakingUser { ServicesOSSAdapter.GetAccountIdRegistry().GetIdValue(OpParams.SessionCreationParameters.LocalAccountId).ToSharedRef() };
 
-		for (const TPair<FSchemaAttributeId, FCustomSessionSetting>& Entry : OpParams.SessionMemberData.MemberSettings)
+		for (const TPair<FSchemaAttributeId, FCustomSessionSetting>& Entry : OpParams.SessionCreationParameters.SessionMemberData.MemberSettings)
 		{
 			switch (Entry.Value.Data.VariantType)
 			{
@@ -831,20 +870,20 @@ TOnlineAsyncOpHandle<FStartMatchmaking> FSessionsOSSAdapter::StartMatchmaking(FS
 
 		MatchMakingUsers.Add(NewMatchMakingUser);
 
-		FOnlineSessionSettings NewSessionSettings = BuildV1Settings(OpParams.SessionSettings);
+		FOnlineSessionSettings NewSessionSettings = BuildV1SettingsForCreate(OpParams.SessionCreationParameters);
 
 		/** We build a mock params struct for FindSessions and convert it to a V1 session search */
 		FFindSessions::Params FindSessionsParams;
 		FindSessionsParams.bFindLANSessions = false;
-		FindSessionsParams.Filters = OpParams.SearchFilters;
-		FindSessionsParams.LocalAccountId = OpParams.LocalAccountId;
+		FindSessionsParams.Filters = OpParams.SessionSearchFilters;
+		FindSessionsParams.LocalAccountId = OpParams.SessionCreationParameters.LocalAccountId;
 		FindSessionsParams.MaxResults = 1;
 
-		TSharedRef<FOnlineSessionSearch>& SessionSearch = PendingV1SessionSearchesPerUser.Emplace(OpParams.LocalAccountId, BuildV1SessionSearch(FindSessionsParams));
+		TSharedRef<FOnlineSessionSearch>& SessionSearch = PendingV1SessionSearchesPerUser.Emplace(OpParams.SessionCreationParameters.LocalAccountId, BuildV1SessionSearch(FindSessionsParams));
 
 		// QUESTION: StartMatchmaking causes both OnMatchmakingComplete and OnStartMatchmakingComplete to trigger.
 		// The first triggers the second so I used the latter, but the interface header did not specify if one should be used. Is this the correct one?
-		SessionsInterface->StartMatchmaking(MatchMakingUsers, OpParams.SessionName, NewSessionSettings, SessionSearch, *MakeDelegateAdapter(this, [this, WeakOp = Op.AsWeak()](FName SessionName, const ::FOnlineError& ErrorDetails, const FSessionMatchmakingResults& Results) mutable
+		SessionsInterface->StartMatchmaking(MatchMakingUsers, OpParams.SessionCreationParameters.SessionName, NewSessionSettings, SessionSearch, *MakeDelegateAdapter(this, [this, WeakOp = Op.AsWeak()](FName SessionName, const ::FOnlineError& ErrorDetails, const FSessionMatchmakingResults& Results) mutable
 		{
 			if (TOnlineAsyncOpPtr<FStartMatchmaking> Op = WeakOp.Pin())
 			{
@@ -858,11 +897,11 @@ TOnlineAsyncOpHandle<FStartMatchmaking> FSessionsOSSAdapter::StartMatchmaking(FS
 
 				TSharedRef<FSessionCommon> V2Session = BuildV2Session(SessionsInterface->GetNamedSession(SessionName));
 
-				AddSessionWithReferences(V2Session, OpParams.SessionName, OpParams.LocalAccountId, V2Session->SessionSettings.bPresenceEnabled);
+				AddSessionWithReferences(V2Session, OpParams.SessionCreationParameters.SessionName, OpParams.SessionCreationParameters.LocalAccountId, OpParams.SessionCreationParameters.bPresenceEnabled);
 
 				Op->SetResult({ });
 
-				FSessionJoined Event{ { Op->GetParams().LocalAccountId }, V2Session->SessionId };
+				FSessionJoined Event{ { Op->GetParams().SessionCreationParameters.LocalAccountId }, V2Session->GetSessionId()};
 
 				SessionEvents.OnSessionJoined.Broadcast(Event);
 
@@ -933,7 +972,7 @@ TOnlineAsyncOpHandle<FJoinSession> FSessionsOSSAdapter::JoinSession(FJoinSession
 
 				TSharedRef<const ISession> FoundSession = GetSessionByIdResult.GetOkValue().Session;
 
-				AddSessionReferences(FoundSession->GetSessionId(), OpParams.SessionName, OpParams.LocalAccountId, FoundSession->GetSessionSettings().bPresenceEnabled);
+				AddSessionReferences(FoundSession->GetSessionId(), OpParams.SessionName, OpParams.LocalAccountId, OpParams.bPresenceEnabled);
 
 				// After successfully joining a session, we'll remove all related invites if any are found
 				ClearSessionInvitesForSession(OpParams.LocalAccountId, FoundSession->GetSessionId());
@@ -1150,74 +1189,119 @@ TOnlineResult<FGetResolvedConnectString> FSessionsOSSAdapter::GetResolvedConnect
 	}
 }
 
-FOnlineSessionSettings FSessionsOSSAdapter::BuildV1Settings(const FSessionSettings& InSessionSettings) const
+FOnlineSessionSettings FSessionsOSSAdapter::BuildV1SettingsForCreate(const FCreateSession::Params& Params) const
 {
 	FOnlineSessionSettings Result;
 
-	Result.bAllowInvites = InSessionSettings.bAllowNewMembers;
-	Result.bAllowJoinInProgress = InSessionSettings.bAllowNewMembers;
-	Result.bAllowJoinViaPresence = InSessionSettings.bAllowNewMembers;
-	Result.bAllowJoinViaPresenceFriendsOnly = InSessionSettings.bAllowNewMembers;
-	Result.Settings.Add(OSS_ADAPTER_SESSIONS_ALLOW_SANCTIONED_PLAYERS, InSessionSettings.bAllowSanctionedPlayers);
-	Result.bAntiCheatProtected = InSessionSettings.bAntiCheatProtected;
-	Result.bIsDedicated = InSessionSettings.bIsDedicatedServerSession;
-	Result.bIsLANMatch = InSessionSettings.bIsLANSession;
-	Result.bShouldAdvertise = InSessionSettings.bAllowNewMembers;
-	if (const FCustomSessionSetting* BuildUniqueId = InSessionSettings.CustomSettings.Find(OSS_ADAPTER_SESSIONS_BUILD_UNIQUE_ID))
+	Result.bAllowInvites = Params.SessionSettings.bAllowNewMembers;
+	Result.bAllowJoinInProgress = Params.SessionSettings.bAllowNewMembers;
+	Result.bAllowJoinViaPresence = Params.SessionSettings.bAllowNewMembers;
+	Result.bAllowJoinViaPresenceFriendsOnly = Params.SessionSettings.bAllowNewMembers;
+	Result.Settings.Add(OSS_ADAPTER_SESSIONS_ALLOW_SANCTIONED_PLAYERS, Params.bAllowSanctionedPlayers);
+	Result.bAntiCheatProtected = Params.bAntiCheatProtected;
+	Result.bIsDedicated = Params.bIsDedicatedServerSession;
+	Result.bIsLANMatch = Params.bIsLANSession;
+	Result.bShouldAdvertise = Params.SessionSettings.bAllowNewMembers;
+	if (const FCustomSessionSetting* BuildUniqueId = Params.SessionSettings.CustomSettings.Find(OSS_ADAPTER_SESSIONS_BUILD_UNIQUE_ID))
 	{
 		Result.BuildUniqueId = BuildUniqueId->Data.GetInt64();
 	}
-	if (const FCustomSessionSetting* BuildUniqueId = InSessionSettings.CustomSettings.Find(OSS_ADAPTER_SESSIONS_USE_LOBBIES_IF_AVAILABLE))
+	if (const FCustomSessionSetting* BuildUniqueId = Params.SessionSettings.CustomSettings.Find(OSS_ADAPTER_SESSIONS_USE_LOBBIES_IF_AVAILABLE))
 	{
 		Result.bUseLobbiesIfAvailable = BuildUniqueId->Data.GetBoolean();
 	}
-	if (const FCustomSessionSetting* BuildUniqueId = InSessionSettings.CustomSettings.Find(OSS_ADAPTER_SESSIONS_USE_LOBBIES_VOICE_CHAT_IF_AVAILABLE))
+	if (const FCustomSessionSetting* BuildUniqueId = Params.SessionSettings.CustomSettings.Find(OSS_ADAPTER_SESSIONS_USE_LOBBIES_VOICE_CHAT_IF_AVAILABLE))
 	{
 		Result.bUseLobbiesVoiceChatIfAvailable = BuildUniqueId->Data.GetBoolean();
 	}
-	Result.bUsesPresence = InSessionSettings.bPresenceEnabled;
-	if (const FCustomSessionSetting* BuildUniqueId = InSessionSettings.CustomSettings.Find(OSS_ADAPTER_SESSIONS_USES_STATS))
+
+	Result.bUsesPresence = Params.bPresenceEnabled;
+
+	if (const FCustomSessionSetting* BuildUniqueId = Params.SessionSettings.CustomSettings.Find(OSS_ADAPTER_SESSIONS_USES_STATS))
 	{
 		Result.bUsesStats = BuildUniqueId->Data.GetBoolean();
 	}
-	Result.NumPrivateConnections = InSessionSettings.NumMaxPrivateConnections;
-	Result.NumPublicConnections = InSessionSettings.NumMaxPublicConnections;
-	Result.Settings.Add(OSS_ADAPTER_SESSIONS_SCHEMA_NAME, InSessionSettings.SchemaName.ToString());
-	Result.SessionIdOverride = InSessionSettings.SessionIdOverride;
+	Result.NumPrivateConnections = Params.SessionSettings.NumMaxConnections;
+	Result.NumPublicConnections = Params.SessionSettings.NumMaxConnections;
+	Result.Settings.Add(OSS_ADAPTER_SESSIONS_SCHEMA_NAME, Params.SessionSettings.SchemaName.ToString());
+	Result.SessionIdOverride = Params.SessionIdOverride;
 
-	Result.Settings = GetV1SessionSettings(InSessionSettings.CustomSettings);
+	Result.Settings = GetV1SessionSettings(Params.SessionSettings.CustomSettings);
+
+	FOnlineServicesOSSAdapter& ServicesOSSAdapter = static_cast<FOnlineServicesOSSAdapter&>(Services);
+
+	Result.MemberSettings.Emplace(ServicesOSSAdapter.GetAccountIdRegistry().GetIdValue(Params.LocalAccountId).ToSharedRef(), GetV1SessionSettings(Params.SessionMemberData.MemberSettings)); // TODO: Pending SchemaVariant work
+
+	return Result;
+}
+
+FOnlineSessionSettings FSessionsOSSAdapter::BuildV1SettingsForUpdate(const FAccountId& LocalAccountId, const TSharedRef<ISession>& Session) const
+{
+	FOnlineSessionSettings Result;
+
+	Result.bAllowInvites = Session->GetSessionSettings().bAllowNewMembers;
+	Result.bAllowJoinInProgress = Session->GetSessionSettings().bAllowNewMembers;
+	Result.bAllowJoinViaPresence = Session->GetSessionSettings().bAllowNewMembers;
+	Result.bAllowJoinViaPresenceFriendsOnly = Session->GetSessionSettings().bAllowNewMembers;
+	Result.Settings.Add(OSS_ADAPTER_SESSIONS_ALLOW_SANCTIONED_PLAYERS, Session->GetSessionInfo().bAllowSanctionedPlayers);
+	Result.bAntiCheatProtected = Session->GetSessionInfo().bAntiCheatProtected;
+	Result.bIsDedicated = Session->GetSessionInfo().bIsDedicatedServerSession;
+	Result.bIsLANMatch = Session->GetSessionInfo().bIsLANSession;
+	Result.bShouldAdvertise = Session->GetSessionSettings().bAllowNewMembers;
+	if (const FCustomSessionSetting* BuildUniqueId = Session->GetSessionSettings().CustomSettings.Find(OSS_ADAPTER_SESSIONS_BUILD_UNIQUE_ID))
+	{
+		Result.BuildUniqueId = BuildUniqueId->Data.GetInt64();
+	}
+	if (const FCustomSessionSetting* BuildUniqueId = Session->GetSessionSettings().CustomSettings.Find(OSS_ADAPTER_SESSIONS_USE_LOBBIES_IF_AVAILABLE))
+	{
+		Result.bUseLobbiesIfAvailable = BuildUniqueId->Data.GetBoolean();
+	}
+	if (const FCustomSessionSetting* BuildUniqueId = Session->GetSessionSettings().CustomSettings.Find(OSS_ADAPTER_SESSIONS_USE_LOBBIES_VOICE_CHAT_IF_AVAILABLE))
+	{
+		Result.bUseLobbiesVoiceChatIfAvailable = BuildUniqueId->Data.GetBoolean();
+	}
+
+	Result.bUsesPresence = false;
+	TOnlineResult<FIsPresenceSession> IsPresenceSessionResult = IsPresenceSession({ LocalAccountId, Session->GetSessionId() });
+	if (IsPresenceSessionResult.IsOk())
+	{
+		Result.bUsesPresence = IsPresenceSessionResult.GetOkValue().bIsPresenceSession;
+	}
+
+	if (const FCustomSessionSetting* BuildUniqueId = Session->GetSessionSettings().CustomSettings.Find(OSS_ADAPTER_SESSIONS_USES_STATS))
+	{
+		Result.bUsesStats = BuildUniqueId->Data.GetBoolean();
+	}
+	Result.NumPrivateConnections = Session->GetSessionSettings().NumMaxConnections;
+	Result.NumPublicConnections = Session->GetSessionSettings().NumMaxConnections;
+	Result.Settings.Add(OSS_ADAPTER_SESSIONS_SCHEMA_NAME, Session->GetSessionSettings().SchemaName.ToString());
+	Result.SessionIdOverride = Session->GetSessionInfo().SessionIdOverride;
+
+	Result.Settings = GetV1SessionSettings(Session->GetSessionSettings().CustomSettings);
 	
 	FOnlineServicesOSSAdapter& ServicesOSSAdapter = static_cast<FOnlineServicesOSSAdapter&>(Services);
 
-	for (const TPair<FAccountId, FSessionMember>& Entry : InSessionSettings.SessionMembers)
+	for (const TPair<FAccountId, FSessionMember>& Entry : Session->GetSessionMembers())
 	{
-		Result.MemberSettings.Emplace(ServicesOSSAdapter.GetAccountIdRegistry().GetIdValue(Entry.Key).ToSharedRef(), ::FSessionSettings()); // TODO: Pending SchemaVariant work
+		Result.MemberSettings.Emplace(ServicesOSSAdapter.GetAccountIdRegistry().GetIdValue(Entry.Key).ToSharedRef(), GetV1SessionSettings(Entry.Value.MemberSettings)); // TODO: Pending SchemaVariant work
 	}
 
 	return Result;
 }
 
-void FSessionsOSSAdapter::WriteV2SessionSettingsFromV1Session(const FOnlineSession* InSession, FSessionSettings& OutSettings) const
+void FSessionsOSSAdapter::WriteV2SessionSettingsFromV1Session(const FOnlineSession* InSession, TSharedRef<FSessionOSSAdapter>& OutSession) const
 {
-	InSession->SessionSettings.Settings.FindChecked(OSS_ADAPTER_SESSIONS_ALLOW_SANCTIONED_PLAYERS).Data.GetValue(OutSettings.bAllowSanctionedPlayers);
-	OutSettings.bAntiCheatProtected = InSession->SessionSettings.bAntiCheatProtected;
-	OutSettings.bPresenceEnabled = InSession->SessionSettings.bUsesPresence;
-	OutSettings.bIsDedicatedServerSession = InSession->SessionSettings.bIsDedicated;
-	OutSettings.bIsLANSession = InSession->SessionSettings.bIsLANMatch;
-	OutSettings.NumMaxPrivateConnections = InSession->SessionSettings.NumPrivateConnections;
-	OutSettings.NumMaxPublicConnections = InSession->SessionSettings.NumPublicConnections;
-	OutSettings.NumOpenPrivateConnections = InSession->NumOpenPrivateConnections;
-	OutSettings.NumOpenPublicConnections = InSession->NumOpenPublicConnections;
+	OutSession->SessionSettings.NumMaxConnections = FMath::Max(InSession->SessionSettings.NumPrivateConnections, InSession->SessionSettings.NumPublicConnections);
+
 	FString SchemaNameStr;
 	InSession->SessionSettings.Settings.FindChecked(OSS_ADAPTER_SESSIONS_SCHEMA_NAME).Data.GetValue(SchemaNameStr);
-	OutSettings.SchemaName = FSchemaId(SchemaNameStr);
-	OutSettings.SessionIdOverride = InSession->SessionSettings.SessionIdOverride;
+	OutSession->SessionSettings.SchemaName = FSchemaId(SchemaNameStr);
 
-	OutSettings.CustomSettings = GetV2SessionSettings(InSession->SessionSettings.Settings);
+	OutSession->SessionSettings.CustomSettings = GetV2SessionSettings(InSession->SessionSettings.Settings);
 
 	FCustomSessionSetting BuildUniqueId;
 	BuildUniqueId.Data.Set((int64)InSession->SessionSettings.BuildUniqueId);
-	OutSettings.CustomSettings.Add(OSS_ADAPTER_SESSIONS_BUILD_UNIQUE_ID, BuildUniqueId);
+	OutSession->SessionSettings.CustomSettings.Add(OSS_ADAPTER_SESSIONS_BUILD_UNIQUE_ID, BuildUniqueId);
 
 	FOnlineServicesOSSAdapter& ServicesOSSAdapter = static_cast<FOnlineServicesOSSAdapter&>(Services);
 
@@ -1228,24 +1312,24 @@ void FSessionsOSSAdapter::WriteV2SessionSettingsFromV1Session(const FOnlineSessi
 		FSessionMember SessionMember;
 		SessionMember.MemberSettings = GetV2SessionSettings(Entry.Value);
 
-		OutSettings.SessionMembers.Emplace(SessionMemberId, SessionMember);
+		OutSession->SessionMembers.Emplace(SessionMemberId, SessionMember);
 	}
 }
 
-void FSessionsOSSAdapter::WriteV2SessionSettingsFromV1NamedSession(const FNamedOnlineSession* InSession, FSessionSettings& OutSettings) const
+void FSessionsOSSAdapter::WriteV2SessionSettingsFromV1NamedSession(const FNamedOnlineSession* InSession, TSharedRef<FSessionOSSAdapter>& OutSession) const
 {
 	bool bPublicJoinable, bFriendJoinable, bInviteOnly, bAllowInvites;
 	InSession->GetJoinability(bPublicJoinable, bFriendJoinable, bInviteOnly, bAllowInvites);
 
-	OutSettings.bAllowNewMembers = bPublicJoinable || bFriendJoinable || bInviteOnly;
-	OutSettings.JoinPolicy = bInviteOnly ? ESessionJoinPolicy::InviteOnly : bFriendJoinable ? ESessionJoinPolicy::FriendsOnly : ESessionJoinPolicy::Public; // We use Public as the default as closed sessions will have bAllowNewMembers as false
+	OutSession->SessionSettings.bAllowNewMembers = bPublicJoinable || bFriendJoinable || bInviteOnly;
+	OutSession->SessionSettings.JoinPolicy = bInviteOnly ? ESessionJoinPolicy::InviteOnly : bFriendJoinable ? ESessionJoinPolicy::FriendsOnly : ESessionJoinPolicy::Public; // We use Public as the default as closed sessions will have bAllowNewMembers as false
 
 	FOnlineServicesOSSAdapter& ServicesOSSAdapter = static_cast<FOnlineServicesOSSAdapter&>(Services);
 
 	// We'll add all the Registered users as SessionMembers with empty data for now
 	for (const FUniqueNetIdRef& RegisteredPlayer : InSession->RegisteredPlayers)
 	{
-		OutSettings.SessionMembers.Emplace(ServicesOSSAdapter.GetAccountIdRegistry().FindOrAddHandle(RegisteredPlayer));
+		OutSession->SessionMembers.Emplace(ServicesOSSAdapter.GetAccountIdRegistry().FindOrAddHandle(RegisteredPlayer));
 	}
 }
 
@@ -1256,18 +1340,17 @@ FOnlineSession FSessionsOSSAdapter::BuildV1Session(const TSharedRef<const ISessi
 
 TSharedRef<FSessionCommon> FSessionsOSSAdapter::BuildV2Session(const FOnlineSession* InSession) const
 {
-	TSharedRef<FSessionOSSAdapter> Session = MakeShared<FSessionOSSAdapter>(*InSession);
-
 	FOnlineServicesOSSAdapter& ServicesOSSAdapter = static_cast<FOnlineServicesOSSAdapter&>(Services);
 
+	TSharedRef<FSessionOSSAdapter> Session = MakeShared<FSessionOSSAdapter>(*InSession, GetSessionIdRegistry().FindOrAddHandle(InSession->SessionInfo->GetSessionId().AsShared()));
+
 	if (const FNamedOnlineSession* NamedInSession = static_cast<const FNamedOnlineSession*>(InSession))
 	{
 		Session->OwnerAccountId = ServicesOSSAdapter.GetAccountIdRegistry().FindOrAddHandle(NamedInSession->LocalOwnerId->AsShared());
-		WriteV2SessionSettingsFromV1NamedSession(NamedInSession, Session->SessionSettings);
+		WriteV2SessionSettingsFromV1NamedSession(NamedInSession, Session);
 	}
 
-	Session->SessionId = GetSessionIdRegistry().FindOrAddHandle(InSession->SessionInfo->GetSessionId().AsShared());
-	WriteV2SessionSettingsFromV1Session(InSession, Session->SessionSettings);
+	WriteV2SessionSettingsFromV1Session(InSession, Session);
 
 	return Session;
 }
diff --git a/Engine/Plugins/Online/Experimental/OnlineServicesOSSAdapter/Source/Public/Online/SessionsOSSAdapter.h b/Engine/Plugins/Online/Experimental/OnlineServicesOSSAdapter/Source/Public/Online/SessionsOSSAdapter.h
--- a/Engine/Plugins/Online/Experimental/OnlineServicesOSSAdapter/Source/Public/Online/SessionsOSSAdapter.h
+++ b/Engine/Plugins/Online/Experimental/OnlineServicesOSSAdapter/Source/Public/Online/SessionsOSSAdapter.h
@@ -38,7 +38,7 @@ class FSessionOSSAdapter : public FSessionCommon
 public:
 	FSessionOSSAdapter() = default;
 	FSessionOSSAdapter(const FSessionOSSAdapter& InSession) = default;
-	FSessionOSSAdapter(const FOnlineSession& InSession);
+	FSessionOSSAdapter(const FOnlineSession& InSession, const FOnlineSessionId& InSessionId);
 
 	static const FSessionOSSAdapter& Cast(const ISession& InSession);
 
@@ -76,12 +76,14 @@ public:
 	TOnlineResult<FGetResolvedConnectString> GetResolvedConnectString(const FGetResolvedConnectString::Params& Params);
 
 private:
-	/** Builds a V1 Session Settings data type from a V2 equivalent */
-	FOnlineSessionSettings BuildV1Settings(const FSessionSettings& InSessionSettings) const;
+	/** Builds a V1 Session Settings data type from V2 session creation parameters */
+	FOnlineSessionSettings BuildV1SettingsForCreate(const FCreateSession::Params& Params) const;
+	/** Builds a V1 Session Settings data type from V2 session data */
+	FOnlineSessionSettings BuildV1SettingsForUpdate(const FAccountId& LocalAccountId, const TSharedRef<ISession>& Session) const;
 	/** Writes V1 Session Settings information into the passed V2 equivalent */
-	void WriteV2SessionSettingsFromV1Session(const FOnlineSession* InSession, FSessionSettings& OutSettings) const;
+	void WriteV2SessionSettingsFromV1Session(const FOnlineSession* InSession, TSharedRef<FSessionOSSAdapter>& OutSession) const;
 	/** Writes V1 Session Settings information contained in Named session type into the passed V2 equivalent */
-	void WriteV2SessionSettingsFromV1NamedSession(const FNamedOnlineSession* InSession, FSessionSettings& OutSettings) const;
+	void WriteV2SessionSettingsFromV1NamedSession(const FNamedOnlineSession* InSession, TSharedRef<FSessionOSSAdapter>& OutSession) const;
 	/** Builds a V1 Session data type from a V2 equivalent */
 	FOnlineSession BuildV1Session(const TSharedRef<const ISession> InSession) const;
 	/** Builds a V2 Session data type from a V1 equivalent */
