- Remove custom parsing from OSSAdapterModule, now there is built in struct support.

diff --git a/Engine/Plugins/Online/Experimental/OnlineServicesOSSAdapter/Source/Private/Online/OnlineServicesOSSAdapterModule.cpp b/Engine/Plugins/Online/Experimental/OnlineServicesOSSAdapter/Source/Private/Online/OnlineServicesOSSAdapterModule.cpp
--- a/Engine/Plugins/Online/Experimental/OnlineServicesOSSAdapter/Source/Private/Online/OnlineServicesOSSAdapterModule.cpp
+++ b/Engine/Plugins/Online/Experimental/OnlineServicesOSSAdapter/Source/Private/Online/OnlineServicesOSSAdapterModule.cpp
@@ -32,8 +32,7 @@ struct FOSSAdapterService
 
 struct FOSSAdapterConfig
 {
-	// TArray<FOSSAdapterService> Services; // use once online config parsing supports arrays of structs
-	TArray<FString> Services;
+	TArray<FOSSAdapterService> Services;
 };
 
 namespace Meta {
@@ -76,64 +75,14 @@ protected:
 	FOSSAdapterService Config;
 };
 
-namespace {
-
-// temporary parsing until online config system has full support for parsing arrays of structs/nested structs
-FString StripQuotes(const FString& Source)
-{
-	if (Source.StartsWith(TEXT("\"")))
-	{
-		return Source.Mid(1, Source.Len() - 2);
-	}
-	return Source;
-}
-
-void ParseServiceConfig(FOSSAdapterService& ServiceConfig, const FString& ServiceConfigString)
-{
-	const TCHAR* Delims[4] = { TEXT("("), TEXT(")"), TEXT("="), TEXT(",") };
-	TArray<FString> Values;
-	ServiceConfigString.ParseIntoArray(Values, Delims, 4, false);
-	for (int32 ValueIndex = 0; ValueIndex + 1 < Values.Num(); ++ValueIndex)
-	{
-		if (Values[ValueIndex].IsEmpty())
-		{
-			continue;
-		}
-
-		if (Values[ValueIndex] == TEXT("Service"))
-		{
-			LexFromString(ServiceConfig.Service, *StripQuotes(Values[ValueIndex + 1]));
-		}
-		else if (Values[ValueIndex] == TEXT("ConfigName"))
-		{
-			ServiceConfig.ConfigName = StripQuotes(Values[ValueIndex + 1]);
-		}
-		else if (Values[ValueIndex] == TEXT("OnlineSubsystem"))
-		{
-			ServiceConfig.OnlineSubsystem = FName(StripQuotes(Values[ValueIndex + 1]));
-		}
-		else if (Values[ValueIndex] == TEXT("Priority"))
-		{
-			ServiceConfig.Priority = FCString::Atoi(*StripQuotes(Values[ValueIndex + 1]));
-		}
-
-		++ValueIndex;
-	}
-}
-
-/* unnamed */ }
-
 void FOnlineServicesOSSAdapterModule::StartupModule()
 {
 	FOnlineConfigProviderGConfig ConfigProvider(GEngineIni);
 	FOSSAdapterConfig Config;
 	if (LoadConfig(ConfigProvider, TEXT("OnlineServices.OSSAdapter"), Config))
 	{
-		for (const FString& ServiceConfigString : Config.Services)
+		for (const FOSSAdapterService& ServiceConfig : Config.Services)
 		{
-			FOSSAdapterService ServiceConfig;
-			ParseServiceConfig(ServiceConfig, ServiceConfigString);
-			
 			if (IOnlineSubsystem::IsEnabled(ServiceConfig.OnlineSubsystem))
 			{
 				FOnlineServicesRegistry::Get().RegisterServicesFactory(ServiceConfig.Service, MakeUnique<FOnlineServicesFactoryOSSAdapter>(ServiceConfig), ServiceConfig.Priority);
@@ -151,11 +100,8 @@ void FOnlineServicesOSSAdapterModule::ShutdownModule()
 	FOSSAdapterConfig Config;
 	if (LoadConfig(ConfigProvider, TEXT("OnlineServices.OSSAdapter"), Config))
 	{
-		for (const FString& ServiceConfigString : Config.Services)
+		for (const FOSSAdapterService& ServiceConfig : Config.Services)
 		{
-			FOSSAdapterService ServiceConfig;
-			ParseServiceConfig(ServiceConfig, ServiceConfigString);
-
 			if (IOnlineSubsystem::IsEnabled(ServiceConfig.OnlineSubsystem))
 			{
 				FOnlineServicesRegistry::Get().UnregisterServicesFactory(ServiceConfig.Service, ServiceConfig.Priority);
