Implementation or removal of pending TODOs in session interface after evaluation (Part 2)
Changed Unknown error codes in Sessions OSS Adapter implementation
Added V1 session update method for session update delegates in Sessions OSS Adapter implementation

diff --git a/Engine/Plugins/Online/Experimental/OnlineServicesOSSAdapter/Source/Private/Online/SessionsOSSAdapter.cpp b/Engine/Plugins/Online/Experimental/OnlineServicesOSSAdapter/Source/Private/Online/SessionsOSSAdapter.cpp
--- a/Engine/Plugins/Online/Experimental/OnlineServicesOSSAdapter/Source/Private/Online/SessionsOSSAdapter.cpp
+++ b/Engine/Plugins/Online/Experimental/OnlineServicesOSSAdapter/Source/Private/Online/SessionsOSSAdapter.cpp
@@ -10,6 +10,7 @@
 #include "Online/AuthOSSAdapter.h"
 #include "Online/UserInfoOSSAdapter.h"
 #include "Online/DelegateAdapter.h"
+#include "Online/ErrorsOSSAdapter.h"
 
 #include "Algo/Find.h"
 
@@ -259,7 +260,7 @@ FSessionOSSAdapter::FSessionOSSAdapter(const FOnlineSession& InSession, const FO
 
 const FSessionOSSAdapter& FSessionOSSAdapter::Cast(const ISession& InSession)
 {
-	// TODO: Check for ::Platform type if that's the Adapter type
+	// TODO: Check for type matching with SessionId
 
 	return static_cast<const FSessionOSSAdapter&>(InSession);
 }
@@ -318,30 +319,27 @@ void FSessionsOSSAdapter::Initialize()
 
 	SessionsInterface->AddOnSessionParticipantRemovedDelegate_Handle(FOnSessionParticipantRemovedDelegate::CreateLambda([this](FName SessionName, const FUniqueNetId& TargetUniqueNetId)
 	{
-		// TODO: Method to update the interface session with the latest OSS Adapter session data
-
-		const FOnlineServicesOSSAdapter& ServicesOSSAdapter = static_cast<FOnlineServicesOSSAdapter&>(Services);
-		const FUniqueNetIdRef TargetUniqueNetIdRef = TargetUniqueNetId.AsShared();
-		const FAccountId TargetAccountId = ServicesOSSAdapter.GetAccountIdRegistry().FindOrAddHandle(TargetUniqueNetIdRef);
+		// We won't update a session that can't be retrieved from the OSS or doesn't exist in OnlineServices anymore
+		if (UpdateV2Session(SessionName))
+		{
+			const FOnlineServicesOSSAdapter& ServicesOSSAdapter = static_cast<FOnlineServicesOSSAdapter&>(Services);
+			const FUniqueNetIdRef TargetUniqueNetIdRef = TargetUniqueNetId.AsShared();
+			const FAccountId TargetAccountId = ServicesOSSAdapter.GetAccountIdRegistry().FindOrAddHandle(TargetUniqueNetIdRef);
 
-		FSessionUpdate SessionUpdate;
-		SessionUpdate.RemovedSessionMembers.Add(TargetAccountId);
+			FSessionUpdate SessionUpdate;
+			SessionUpdate.RemovedSessionMembers.Add(TargetAccountId);
 
-		const FSessionUpdated Event { SessionName , SessionUpdate };
+			const FSessionUpdated Event{ SessionName , SessionUpdate };
 
-		SessionEvents.OnSessionUpdated.Broadcast(Event);
+			SessionEvents.OnSessionUpdated.Broadcast(Event);
+		}
 	}));
 
 	SessionsInterface->AddOnSessionParticipantsChangeDelegate_Handle(FOnSessionParticipantsChangeDelegate::CreateLambda([this](FName SessionName, const FUniqueNetId& TargetUniqueNetId, bool bJoined)
 	{
-		// TODO: Method to update the interface session with the latest OSS Adapter session data
-
-		// We won't transmit events for a session or member that doesn't exist
-		const TOnlineResult<FGetSessionByName> GetSessionByNameResult = GetSessionByName({ SessionName });
-		if (GetSessionByNameResult.IsOk())
+		// We won't update a session that can't be retrieved from the OSS or doesn't exist in OnlineServices anymore
+		if (UpdateV2Session(SessionName))
 		{
-			const TSharedRef<const ISession>& FoundSession = GetSessionByNameResult.GetOkValue().Session;
-
 			const FOnlineServicesOSSAdapter& ServicesOSSAdapter = static_cast<FOnlineServicesOSSAdapter&>(Services);
 			const FUniqueNetIdRef TargetUniqueNetIdRef = TargetUniqueNetId.AsShared();
 			const FAccountId TargetAccountId = ServicesOSSAdapter.GetAccountIdRegistry().FindOrAddHandle(TargetUniqueNetIdRef);
@@ -356,7 +354,7 @@ void FSessionsOSSAdapter::Initialize()
 				SessionUpdate.RemovedSessionMembers.Emplace(TargetAccountId);
 			}
 
-			const FSessionUpdated Event { SessionName , SessionUpdate };
+			const FSessionUpdated Event{ SessionName , SessionUpdate };
 
 			SessionEvents.OnSessionUpdated.Broadcast(Event);
 		}
@@ -390,7 +388,7 @@ TFuture<TOnlineResult<FCreateSession>> FSessionsOSSAdapter::CreateSessionImpl(co
 	{
 		if (!bWasSuccessful)
 		{
-			Promise.EmplaceValue(Errors::Unknown());
+			Promise.EmplaceValue(Errors::RequestFailure());
 			return;
 		}
 
@@ -418,7 +416,7 @@ TFuture<TOnlineResult<FUpdateSessionSettings>> FSessionsOSSAdapter::UpdateSessio
 	{
 		if (!bWasSuccessful)
 		{
-			Promise.EmplaceValue(Errors::Unknown());
+			Promise.EmplaceValue(Errors::RequestFailure());
 			return;
 		}
 
@@ -472,7 +470,7 @@ TFuture<TOnlineResult<FLeaveSession>> FSessionsOSSAdapter::LeaveSessionImpl(cons
 	{
 		if (!bWasSuccessful)
 		{
-			Promise.EmplaceValue(Errors::Unknown());
+			Promise.EmplaceValue(Errors::RequestFailure());
 			SessionsInterface->ClearOnDestroySessionCompleteDelegates(this);
 			return;
 		}
@@ -550,7 +548,7 @@ TFuture<TOnlineResult<FFindSessions>> FSessionsOSSAdapter::FindSessionsImpl(cons
 				}
 				else
 				{
-					Promise.EmplaceValue(Errors::Unknown());
+					Promise.EmplaceValue(Errors::RequestFailure());
 				}
 
 				CurrentSessionSearchPromisesUserMap.Remove(Params.LocalAccountId);
@@ -573,7 +571,7 @@ TFuture<TOnlineResult<FFindSessions>> FSessionsOSSAdapter::FindSessionsImpl(cons
 
 				if (!bWasSuccessful)
 				{
-					Promise.EmplaceValue(Errors::Unknown());
+					Promise.EmplaceValue(Errors::RequestFailure());
 					return;
 				}
 
@@ -610,7 +608,7 @@ TFuture<TOnlineResult<FFindSessions>> FSessionsOSSAdapter::FindSessionsImpl(cons
 			}
 			else
 			{
-				Promise.EmplaceValue(Errors::Unknown());
+				Promise.EmplaceValue(Errors::RequestFailure());
 			}
 
 			CurrentSessionSearchPromisesUserMap.Remove(Params.LocalAccountId);
@@ -648,7 +646,7 @@ TFuture<TOnlineResult<FFindSessions>> FSessionsOSSAdapter::FindSessionsImpl(cons
 			}
 			else
 			{
-				Promise.EmplaceValue(Errors::Unknown());
+				Promise.EmplaceValue(Errors::RequestFailure());
 			}
 
 			CurrentSessionSearchPromisesUserMap.Remove(Params.LocalAccountId);
@@ -689,7 +687,7 @@ TFuture<TOnlineResult<FStartMatchmaking>> FSessionsOSSAdapter::StartMatchmakingI
 	{
 		if (!ErrorDetails.bSucceeded)
 		{
-			Promise.EmplaceValue(Errors::Unknown()); // TODO: Proper error parsing
+			Promise.EmplaceValue(Errors::FromOssError(ErrorDetails));
 			return;
 		}
 
@@ -724,7 +722,7 @@ TFuture<TOnlineResult<FJoinSession>> FSessionsOSSAdapter::JoinSessionImpl(const
 	{
 		if (Result != EOnJoinSessionCompleteResult::Success)
 		{
-			Promise.EmplaceValue(Errors::Unknown());
+			Promise.EmplaceValue(Errors::RequestFailure());
 			return;
 		}
 
@@ -792,7 +790,7 @@ TFuture<TOnlineResult<FAddSessionMember>> FSessionsOSSAdapter::AddSessionMemberI
 			}
 			else
 			{
-				Promise.EmplaceValue(Errors::Unknown());
+				Promise.EmplaceValue(Errors::RequestFailure());
 			}
 		});
 
@@ -823,7 +821,7 @@ TFuture<TOnlineResult<FRemoveSessionMember>> FSessionsOSSAdapter::RemoveSessionM
 			}
 			else
 			{
-				Promise.EmplaceValue(Errors::Unknown());
+				Promise.EmplaceValue(Errors::RequestFailure());
 			}
 		});
 
@@ -889,7 +887,7 @@ TOnlineResult<FGetResolvedConnectString> FSessionsOSSAdapter::GetResolvedConnect
 			}
 			else
 			{
-				return TOnlineResult<FGetResolvedConnectString>(Errors::Unknown());
+				return TOnlineResult<FGetResolvedConnectString>(Errors::RequestFailure());
 			}
 		}
 		else
@@ -904,6 +902,27 @@ TOnlineResult<FGetResolvedConnectString> FSessionsOSSAdapter::GetResolvedConnect
 	}
 }
 
+bool FSessionsOSSAdapter::UpdateV2Session(const FName& SessionName)
+{
+	bool bResult = false;
+
+	if (FNamedOnlineSession* NamedOnlineSession = SessionsInterface->GetNamedSession(SessionName))
+	{
+		TOnlineResult<FGetMutableSessionByName> GetMutableSessionByNameResult = GetMutableSessionByName({ SessionName });
+		if (GetMutableSessionByNameResult.IsOk())
+		{
+			// We update our local session
+			TSharedRef<FSessionCommon>& FoundSession = GetMutableSessionByNameResult.GetOkValue().Session;
+
+			FoundSession = BuildV2Session(NamedOnlineSession);
+
+			bResult = true;
+		}
+	}
+
+	return bResult;
+}
+
 FOnlineSessionSettings FSessionsOSSAdapter::BuildV1SettingsForCreate(const FCreateSession::Params& Params) const
 {
 	FOnlineSessionSettings Result;
diff --git a/Engine/Plugins/Online/Experimental/OnlineServicesOSSAdapter/Source/Public/Online/SessionsOSSAdapter.h b/Engine/Plugins/Online/Experimental/OnlineServicesOSSAdapter/Source/Public/Online/SessionsOSSAdapter.h
--- a/Engine/Plugins/Online/Experimental/OnlineServicesOSSAdapter/Source/Public/Online/SessionsOSSAdapter.h
+++ b/Engine/Plugins/Online/Experimental/OnlineServicesOSSAdapter/Source/Public/Online/SessionsOSSAdapter.h
@@ -78,6 +78,9 @@ public:
 	TOnlineResult<FGetResolvedConnectString> GetResolvedConnectString(const FGetResolvedConnectString::Params& Params);
 
 private:
+	/** Updates a local session's state by retrieving the equivalent from the OSS */
+	bool UpdateV2Session(const FName& SessionName);
+
 	/** Builds a V1 Session Settings data type from V2 session creation parameters */
 	FOnlineSessionSettings BuildV1SettingsForCreate(const FCreateSession::Params& Params) const;
 	/** Builds a V1 Session Settings data type from V2 session data */
