Added missing logs to error cases in OSSV2 Sessions

diff --git a/Engine/Plugins/Online/OnlineServicesEOSGS/Source/Private/Online/SessionsEOSGS.cpp b/Engine/Plugins/Online/OnlineServicesEOSGS/Source/Private/Online/SessionsEOSGS.cpp
--- a/Engine/Plugins/Online/OnlineServicesEOSGS/Source/Private/Online/SessionsEOSGS.cpp
+++ b/Engine/Plugins/Online/OnlineServicesEOSGS/Source/Private/Online/SessionsEOSGS.cpp
@@ -943,6 +943,7 @@ TFuture<TOnlineResult<FLeaveSession>> FSessionsEOSGS::LeaveSessionImpl(const FLe
 		{
 			if (Result->ResultCode != EOS_EResult::EOS_Success)
 			{
+				UE_LOG(LogTemp, Warning, TEXT("EOS_Sessions_DestroySession failed with result [%s]"), *LexToString(Result->ResultCode));
 				Promise.EmplaceValue(Errors::FromEOSResult(Result->ResultCode));
 				return;
 			}
@@ -1115,6 +1116,7 @@ TFuture<TOnlineResult<FFindSessions>> FSessionsEOSGS::FindSessionsImpl(const FFi
 
 			if (FindCallbackInfoResult->ResultCode != EOS_EResult::EOS_Success)
 			{
+				UE_LOG(LogTemp, Warning, TEXT("EOS_SessionSearch_Find failed with result [%s]"), *LexToString(FindCallbackInfoResult->ResultCode));
 				Promise.EmplaceValue(Errors::FromEOSResult(FindCallbackInfoResult->ResultCode));
 				return;
 			}
@@ -1251,6 +1253,7 @@ TFuture<TOnlineResult<FJoinSession>> FSessionsEOSGS::JoinSessionImpl(const FJoin
 		{
 			if (Result->ResultCode != EOS_EResult::EOS_Success)
 			{
+				UE_LOG(LogTemp, Warning, TEXT("EOS_Sessions_JoinSession failed with result [%s]"), *LexToString(Result->ResultCode));
 				Promise.EmplaceValue(Errors::FromEOSResult(Result->ResultCode));
 				return;
 			}
@@ -1495,6 +1498,7 @@ TFuture<TOnlineResult<FRejectSessionInvite>> FSessionsEOSGS::RejectSessionInvite
 		{
 			if (Result->ResultCode != EOS_EResult::EOS_Success)
 			{
+				UE_LOG(LogTemp, Warning, TEXT("EOS_Sessions_RejectInvite failed with result [%s]"), *LexToString(Result->ResultCode));
 				Promise.EmplaceValue(Errors::FromEOSResult(Result->ResultCode));
 				return;
 			}
@@ -1546,6 +1550,7 @@ TFuture<TOnlineResult<FAddSessionMember>> FSessionsEOSGS::AddSessionMemberImpl(c
 						{
 							if (Result->ResultCode != EOS_EResult::EOS_Success)
 							{
+								UE_LOG(LogTemp, Warning, TEXT("EOS_Sessions_RegisterPlayers failed with result [%s]"), *LexToString(Result->ResultCode));
 								Promise.EmplaceValue(Errors::FromEOSResult(Result->ResultCode));
 								return;
 							}
@@ -1598,6 +1603,7 @@ TFuture<TOnlineResult<FRemoveSessionMember>> FSessionsEOSGS::RemoveSessionMember
 						{
 							if (Result->ResultCode != EOS_EResult::EOS_Success)
 							{
+								UE_LOG(LogTemp, Warning, TEXT("EOS_Sessions_UnregisterPlayers failed with result [%s]"), *LexToString(Result->ResultCode));
 								Promise.EmplaceValue(Errors::FromEOSResult(Result->ResultCode));
 								return;
 							}
diff --git a/Engine/Plugins/Online/OnlineServicesOSSAdapter/Source/Private/Online/SessionsOSSAdapter.cpp b/Engine/Plugins/Online/OnlineServicesOSSAdapter/Source/Private/Online/SessionsOSSAdapter.cpp
--- a/Engine/Plugins/Online/OnlineServicesOSSAdapter/Source/Private/Online/SessionsOSSAdapter.cpp
+++ b/Engine/Plugins/Online/OnlineServicesOSSAdapter/Source/Private/Online/SessionsOSSAdapter.cpp
@@ -501,12 +501,14 @@ TFuture<TOnlineResult<FFindSessions>> FSessionsOSSAdapter::FindSessionsImpl(cons
 	const FUniqueNetIdRef LocalAccountId = Services.Get<FAuthOSSAdapter>()->GetUniqueNetId(Params.LocalAccountId).ToSharedRef();
 	if (!LocalAccountId->IsValid())
 	{
+		UE_LOG(LogTemp, Verbose, TEXT("[FSessionsOSSAdapter::FindSessionsImpl] UniqueId not found for LocalAccountId %s"), *ToLogString(Params.LocalAccountId));
 		Promise.EmplaceValue(Errors::InvalidUser());
 		return Future;
 	}
 
 	if (PendingV1SessionSearchesPerUser.Contains(Params.LocalAccountId))
 	{
+		UE_LOG(LogTemp, Verbose, TEXT("[FSessionsOSSAdapter::FindSessionsImpl] Session search already in progress for LocalAccountId %s"), *ToLogString(Params.LocalAccountId));
 		Promise.EmplaceValue(Errors::AlreadyPending());
 		return Future;
 	}
@@ -575,6 +577,7 @@ TFuture<TOnlineResult<FFindSessions>> FSessionsOSSAdapter::FindSessionsImpl(cons
 		}
 		else
 		{
+			UE_LOG(LogTemp, Verbose, TEXT("[FSessionsOSSAdapter::FindSessionsImpl] UniqueNetId could not be found for SessionId %s"), *ToLogString(*Params.SessionId));
 			Promise.EmplaceValue(Errors::InvalidParams());
 			return Future;
 		}
@@ -706,6 +709,7 @@ TFuture<TOnlineResult<FJoinSession>> FSessionsOSSAdapter::JoinSessionImpl(const
 	const FUniqueNetIdRef LocalAccountId = Services.Get<FAuthOSSAdapter>()->GetUniqueNetId(Params.LocalAccountId).ToSharedRef();
 	if (!LocalAccountId->IsValid())
 	{
+		UE_LOG(LogTemp, Verbose, TEXT("[FSessionsOSSAdapter::JoinSessionImpl] UniqueId not found for LocalAccountId %s"), *ToLogString(Params.LocalAccountId));
 		Promise.EmplaceValue(Errors::InvalidUser());
 		return Future;
 	}
@@ -841,6 +845,7 @@ TFuture<TOnlineResult<FSendSessionInvite>> FSessionsOSSAdapter::SendSessionInvit
 	const FUniqueNetIdRef LocalAccountId = Auth->GetUniqueNetId(Params.LocalAccountId).ToSharedRef();
 	if (!LocalAccountId->IsValid())
 	{
+		UE_LOG(LogTemp, Verbose, TEXT("[FSessionsOSSAdapter::SendSessionInviteImpl] UniqueId not found for LocalAccountId %s"), *ToLogString(Params.LocalAccountId));
 		Promise.EmplaceValue(Errors::InvalidUser());
 		return Future;
 	}
@@ -891,6 +896,7 @@ TOnlineResult<FGetResolvedConnectString> FSessionsOSSAdapter::GetResolvedConnect
 	else
 	{
 		// No valid session id set
+		UE_LOG(LogTemp, Verbose, TEXT("[FSessionsOSSAdapter::GetResolvedConnectString] Invalid SessionId %s"), *ToLogString(Params.SessionId));
 		return TOnlineResult<FGetResolvedConnectString>(Errors::InvalidParams());
 	}
 }
