Refactor of implementation dependent checks in Sessions Methods

diff --git a/Engine/Plugins/Online/OnlineServices/Source/OnlineServicesCommon/Public/Online/SessionsCommon.h b/Engine/Plugins/Online/OnlineServices/Source/OnlineServicesCommon/Public/Online/SessionsCommon.h
--- a/Engine/Plugins/Online/OnlineServices/Source/OnlineServicesCommon/Public/Online/SessionsCommon.h
+++ b/Engine/Plugins/Online/OnlineServices/Source/OnlineServicesCommon/Public/Online/SessionsCommon.h
@@ -285,26 +285,46 @@ protected:
 	// FSessionsCommon internal interface
 
 	virtual TFuture<TOnlineResult<FCreateSession>> CreateSessionImpl(const FCreateSession::Params& Params);
+	virtual TOptional<FOnlineError> CheckParams(const FCreateSession::Params& Params) const;
+	virtual TOptional<FOnlineError> CheckState(const FCreateSession::Params& Params) const;
 
 	virtual TFuture<TOnlineResult<FUpdateSessionSettings>> UpdateSessionSettingsImpl(const FUpdateSessionSettings::Params& Params);
+	virtual TOptional<FOnlineError> CheckParams(const FUpdateSessionSettings::Params& Params) const;
+	virtual TOptional<FOnlineError> CheckState(const FUpdateSessionSettings::Params& Params) const;
 
 	virtual TFuture<TOnlineResult<FLeaveSession>> LeaveSessionImpl(const FLeaveSession::Params& Params);
+	virtual TOptional<FOnlineError> CheckParams(const FLeaveSession::Params& Params) const;
+	virtual TOptional<FOnlineError> CheckState(const FLeaveSession::Params& Params) const;
 
 	virtual TFuture<TOnlineResult<FFindSessions>> FindSessionsImpl(const FFindSessions::Params& Params);
+	virtual TOptional<FOnlineError> CheckParams(const FFindSessions::Params& Params) const;
+	virtual TOptional<FOnlineError> CheckState(const FFindSessions::Params& Params) const;
 
 	virtual TFuture<TOnlineResult<FStartMatchmaking>> StartMatchmakingImpl(const FStartMatchmaking::Params& Params);
+	virtual TOptional<FOnlineError> CheckParams(const FStartMatchmaking::Params& Params) const;
+	virtual TOptional<FOnlineError> CheckState(const FStartMatchmaking::Params& Params) const;
 
 	virtual TFuture<TOnlineResult<FJoinSession>> JoinSessionImpl(const FJoinSession::Params& Params);
+	virtual TOptional<FOnlineError> CheckParams(const FJoinSession::Params& Params) const;
+	virtual TOptional<FOnlineError> CheckState(const FJoinSession::Params& Params) const;
 
 	virtual TFuture<TOnlineResult<FAddSessionMember>> AddSessionMemberImpl(const FAddSessionMember::Params& Params);
+	virtual TOptional<FOnlineError> CheckParams(const FAddSessionMember::Params& Params) const;
+	virtual TOptional<FOnlineError> CheckState(const FAddSessionMember::Params& Params) const;
 	TOnlineResult<FAddSessionMember> AddSessionMemberInternal(const FAddSessionMember::Params& Params);
 
 	virtual TFuture<TOnlineResult<FRemoveSessionMember>> RemoveSessionMemberImpl(const FRemoveSessionMember::Params& Params);
+	virtual TOptional<FOnlineError> CheckParams(const FRemoveSessionMember::Params& Params) const;
+	virtual TOptional<FOnlineError> CheckState(const FRemoveSessionMember::Params& Params) const;
 	TOnlineResult<FRemoveSessionMember> RemoveSessionMemberInternal(const FRemoveSessionMember::Params& Params);
 
 	virtual TFuture<TOnlineResult<FSendSessionInvite>> SendSessionInviteImpl(const FSendSessionInvite::Params& Params);
+	virtual TOptional<FOnlineError> CheckParams(const FSendSessionInvite::Params& Params) const;
+	virtual TOptional<FOnlineError> CheckState(const FSendSessionInvite::Params& Params) const;
 
 	virtual TFuture<TOnlineResult<FRejectSessionInvite>> RejectSessionInviteImpl(const FRejectSessionInvite::Params& Params);
+	virtual TOptional<FOnlineError> CheckParams(const FRejectSessionInvite::Params& Params) const;
+	virtual TOptional<FOnlineError> CheckState(const FRejectSessionInvite::Params& Params) const;
 
 private:
 	template<typename MethodStruct>
@@ -313,36 +333,6 @@ private:
 	void ClearSessionByName(const FName& SessionName);
 	void ClearSessionById(const FOnlineSessionId& SessionId);
 
-	TOptional<FOnlineError> CheckParams(const FCreateSession::Params& Params) const;
-	TOptional<FOnlineError> CheckState(const FCreateSession::Params& Params) const;
-
-	TOptional<FOnlineError> CheckParams(const FUpdateSessionSettings::Params& Params) const;
-	TOptional<FOnlineError> CheckState(const FUpdateSessionSettings::Params& Params) const;
-
-	TOptional<FOnlineError> CheckParams(const FLeaveSession::Params& Params) const;
-	TOptional<FOnlineError> CheckState(const FLeaveSession::Params& Params) const;
-
-	TOptional<FOnlineError> CheckParams(const FFindSessions::Params& Params) const;
-	TOptional<FOnlineError> CheckState(const FFindSessions::Params& Params) const;
-
-	TOptional<FOnlineError> CheckParams(const FStartMatchmaking::Params& Params) const;
-	TOptional<FOnlineError> CheckState(const FStartMatchmaking::Params& Params) const;
-
-	TOptional<FOnlineError> CheckParams(const FJoinSession::Params& Params) const;
-	TOptional<FOnlineError> CheckState(const FJoinSession::Params& Params) const;
-
-	TOptional<FOnlineError> CheckParams(const FAddSessionMember::Params& Params) const;
-	TOptional<FOnlineError> CheckState(const FAddSessionMember::Params& Params) const;
-
-	TOptional<FOnlineError> CheckParams(const FRemoveSessionMember::Params& Params) const;
-	TOptional<FOnlineError> CheckState(const FRemoveSessionMember::Params& Params) const;
-
-	TOptional<FOnlineError> CheckParams(const FSendSessionInvite::Params& Params) const;
-	TOptional<FOnlineError> CheckState(const FSendSessionInvite::Params& Params) const;
-
-	TOptional<FOnlineError> CheckParams(const FRejectSessionInvite::Params& Params) const;
-	TOptional<FOnlineError> CheckState(const FRejectSessionInvite::Params& Params) const;
-
 protected:
 
 	struct FSessionEvents
diff --git a/Engine/Plugins/Online/OnlineServicesEOSGS/Source/Private/Online/SessionsEOSGS.cpp b/Engine/Plugins/Online/OnlineServicesEOSGS/Source/Private/Online/SessionsEOSGS.cpp
--- a/Engine/Plugins/Online/OnlineServicesEOSGS/Source/Private/Online/SessionsEOSGS.cpp
+++ b/Engine/Plugins/Online/OnlineServicesEOSGS/Source/Private/Online/SessionsEOSGS.cpp
@@ -326,22 +326,6 @@ TFuture<TOnlineResult<FCreateSession>> FSessionsEOSGS::CreateSessionImpl(const F
 	TPromise<TOnlineResult<FCreateSession>> Promise;
 	TFuture<TOnlineResult<FCreateSession>> Future = Promise.GetFuture();
 
-	if (!Params.SessionIdOverride.IsEmpty())
-	{
-		int32 Length = Params.SessionIdOverride.Len();
-
-		if (Length < EOS_SESSIONMODIFICATION_MIN_SESSIONIDOVERRIDE_LENGTH || Length > EOS_SESSIONMODIFICATION_MAX_SESSIONIDOVERRIDE_LENGTH)
-		{
-			UE_LOG(LogTemp, Warning, TEXT("[FSessionsEOSGS::CreateSession] Could not create session with SessionIdOverride [%s] of size [%d]. SessionIdOverride size must be between [%d] and [%d] characters long"), *Params.SessionIdOverride, Length, EOS_SESSIONMODIFICATION_MIN_SESSIONIDOVERRIDE_LENGTH, EOS_SESSIONMODIFICATION_MAX_SESSIONIDOVERRIDE_LENGTH);
-
-			Promise.EmplaceValue(Errors::InvalidParams());
-
-			return Future;
-		}
-	}
-
-	// After all initial checks, we start the session creation operations
-
 	EOS_Sessions_CreateSessionModificationOptions CreateSessionModificationOptions = {};
 	CreateSessionModificationOptions.ApiVersion = EOS_SESSIONS_CREATESESSIONMODIFICATION_API_LATEST;
 	static_assert(EOS_SESSIONS_CREATESESSIONMODIFICATION_API_LATEST == 4, "EOS_Sessions_CreateSessionModificationOptions updated, check new fields");
@@ -425,6 +409,28 @@ TFuture<TOnlineResult<FCreateSession>> FSessionsEOSGS::CreateSessionImpl(const F
 	return Future;
 }
 
+TOptional<FOnlineError> FSessionsEOSGS::CheckState(const FCreateSession::Params& Params) const
+{
+	if (TOptional<FOnlineError> BaseCheck = Super::CheckState(Params))
+	{
+		return BaseCheck;
+	}
+	
+	if (!Params.bIsLANSession && !Params.SessionIdOverride.IsEmpty())
+	{
+		int32 Length = Params.SessionIdOverride.Len();
+
+		if (Length < EOS_SESSIONMODIFICATION_MIN_SESSIONIDOVERRIDE_LENGTH || Length > EOS_SESSIONMODIFICATION_MAX_SESSIONIDOVERRIDE_LENGTH)
+		{
+			UE_LOG(LogTemp, Warning, TEXT("[%s] Could not create session with SessionIdOverride [%s] of size [%d]. SessionIdOverride size must be between [%d] and [%d] characters long"), UTF8_TO_TCHAR(__FUNCTION__), *Params.SessionIdOverride, Length, EOS_SESSIONMODIFICATION_MIN_SESSIONIDOVERRIDE_LENGTH, EOS_SESSIONMODIFICATION_MAX_SESSIONIDOVERRIDE_LENGTH);
+
+			return TOptional<FOnlineError>(Errors::InvalidParams());
+		}
+	}
+
+	return TOptional<FOnlineError>();
+}
+
 void FSessionsEOSGS::SetHostAddress(EOS_HSessionModification& SessionModHandle, FString HostAddress)
 {
 	EOS_SessionModification_SetHostAddressOptions Options = {};
@@ -1214,16 +1220,6 @@ TFuture<TOnlineResult<FJoinSession>> FSessionsEOSGS::JoinSessionImpl(const FJoin
 
 	// EOSGS Sessions
 
-	// We check that the passed session has a valid details handle
-	const FSessionEOSGS& SessionEOSGS = FSessionEOSGS::Cast(*FoundSession);
-	if (!SessionEOSGS.SessionDetailsHandle.IsValid())
-	{
-		UE_LOG(LogTemp, Warning, TEXT("[%s] Could not join session with invalid session details handle in session with id [%s]"), UTF8_TO_TCHAR(__FUNCTION__), *ToLogString(Params.SessionId));
-
-		Promise.EmplaceValue(Errors::InvalidState());
-		return Future;
-	}
-
 	// We start setup for the API call
 	EOS_Sessions_JoinSessionOptions JoinSessionOptions = { };
 	JoinSessionOptions.ApiVersion = EOS_SESSIONS_JOINSESSION_API_LATEST;
@@ -1236,6 +1232,7 @@ TFuture<TOnlineResult<FJoinSession>> FSessionsEOSGS::JoinSessionImpl(const FJoin
 	const FTCHARToUTF8 SessionNameUtf8(*Params.SessionName.ToString());
 	JoinSessionOptions.SessionName = SessionNameUtf8.Get();
 
+	const FSessionEOSGS& SessionEOSGS = FSessionEOSGS::Cast(*FoundSession);
 	JoinSessionOptions.SessionHandle = SessionEOSGS.SessionDetailsHandle->SessionDetailsHandle;
 
 	EOS_Async(EOS_Sessions_JoinSession, SessionsHandle, JoinSessionOptions, 
@@ -1280,6 +1277,35 @@ TFuture<TOnlineResult<FJoinSession>> FSessionsEOSGS::JoinSessionImpl(const FJoin
 	return Future;
 }
 
+TOptional<FOnlineError> FSessionsEOSGS::CheckState(const FJoinSession::Params& Params) const
+{
+	if (TOptional<FOnlineError> BaseCheck = Super::CheckState(Params))
+	{
+		return BaseCheck;
+	}
+
+	TOnlineResult<FGetSessionById> GetSessionByIdResult = GetSessionById({ Params.SessionId });
+	if (GetSessionByIdResult.IsOk())
+	{
+		const TSharedRef<const ISession>& FoundSession = GetSessionByIdResult.GetOkValue().Session;
+
+		// We check that the passed session has a valid details handle
+		const FSessionEOSGS& SessionEOSGS = FSessionEOSGS::Cast(*FoundSession);
+		if (!FoundSession->GetSessionInfo().bIsLANSession && !SessionEOSGS.SessionDetailsHandle.IsValid())
+		{
+			UE_LOG(LogTemp, Warning, TEXT("[%s] Could not join session with invalid session details handle in session with id [%s]"), UTF8_TO_TCHAR(__FUNCTION__), *ToLogString(Params.SessionId));
+
+			return TOptional<FOnlineError>(Errors::InvalidState());
+		}
+	}
+	else
+	{
+		return TOptional<FOnlineError>(GetSessionByIdResult.GetErrorValue());
+	}
+
+	return TOptional<FOnlineError>();
+}
+
 TOnlineAsyncOpHandle<FBuildSessionFromDetailsHandle> FSessionsEOSGS::BuildSessionFromInvite(const FAccountId& LocalAccountId, const FString& InInviteId)
 {
 	EOS_Sessions_CopySessionHandleByInviteIdOptions CopySessionHandleByInviteIdOptions = { };
diff --git a/Engine/Plugins/Online/OnlineServicesEOSGS/Source/Public/Online/SessionsEOSGS.h b/Engine/Plugins/Online/OnlineServicesEOSGS/Source/Public/Online/SessionsEOSGS.h
--- a/Engine/Plugins/Online/OnlineServicesEOSGS/Source/Public/Online/SessionsEOSGS.h
+++ b/Engine/Plugins/Online/OnlineServicesEOSGS/Source/Public/Online/SessionsEOSGS.h
@@ -203,6 +203,9 @@ public:
 	virtual TFuture<TOnlineResult<FRejectSessionInvite>> RejectSessionInviteImpl(const FRejectSessionInvite::Params& Params) override;
 
 protected:
+	virtual TOptional<FOnlineError> CheckState(const FCreateSession::Params& Params) const override;
+	virtual TOptional<FOnlineError> CheckState(const FJoinSession::Params& Params) const override;
+
 	void RegisterEventHandlers();
 	void UnregisterEventHandlers();
 	void HandleSessionInviteReceived(const EOS_Sessions_SessionInviteReceivedCallbackInfo* Data);
diff --git a/Engine/Plugins/Online/OnlineServicesOSSAdapter/Source/Private/Online/SessionsOSSAdapter.cpp b/Engine/Plugins/Online/OnlineServicesOSSAdapter/Source/Private/Online/SessionsOSSAdapter.cpp
--- a/Engine/Plugins/Online/OnlineServicesOSSAdapter/Source/Private/Online/SessionsOSSAdapter.cpp
+++ b/Engine/Plugins/Online/OnlineServicesOSSAdapter/Source/Private/Online/SessionsOSSAdapter.cpp
@@ -498,25 +498,14 @@ TFuture<TOnlineResult<FFindSessions>> FSessionsOSSAdapter::FindSessionsImpl(cons
 	TPromise<TOnlineResult<FFindSessions>> Promise;
 	TFuture<TOnlineResult<FFindSessions>> Future = Promise.GetFuture();
 
-	const FUniqueNetIdRef LocalAccountId = Services.Get<FAuthOSSAdapter>()->GetUniqueNetId(Params.LocalAccountId).ToSharedRef();
-	if (!LocalAccountId->IsValid())
-	{
-		UE_LOG(LogTemp, Verbose, TEXT("[FSessionsOSSAdapter::FindSessionsImpl] UniqueId not found for LocalAccountId %s"), *ToLogString(Params.LocalAccountId));
-		Promise.EmplaceValue(Errors::InvalidUser());
-		return Future;
-	}
-
-	if (PendingV1SessionSearchesPerUser.Contains(Params.LocalAccountId))
-	{
-		UE_LOG(LogTemp, Verbose, TEXT("[FSessionsOSSAdapter::FindSessionsImpl] Session search already in progress for LocalAccountId %s"), *ToLogString(Params.LocalAccountId));
-		Promise.EmplaceValue(Errors::AlreadyPending());
-		return Future;
-	}
-
 	// Before we start the search, we reset the cache and save the promise
 	SearchResultsUserMap.FindOrAdd(Params.LocalAccountId).Reset();
 	CurrentSessionSearchPromisesUserMap.Emplace(Params.LocalAccountId, MoveTemp(Promise));
 
+	const FUniqueNetIdPtr LocalAccountIdPtr = Services.Get<FAuthOSSAdapter>()->GetUniqueNetId(Params.LocalAccountId);
+	check(LocalAccountIdPtr.IsValid());
+	const FUniqueNetIdRef LocalAccountId = LocalAccountIdPtr.ToSharedRef();
+
 	if (Params.SessionId.IsSet())
 	{
 		if (const FUniqueNetIdPtr SessionIdPtr = GetSessionIdRegistry().GetIdValue(*Params.SessionId))
@@ -657,6 +646,31 @@ TFuture<TOnlineResult<FFindSessions>> FSessionsOSSAdapter::FindSessionsImpl(cons
 	return Future;
 }
 
+TOptional<FOnlineError> FSessionsOSSAdapter::CheckState(const FFindSessions::Params& Params) const
+{
+	if (TOptional<FOnlineError> BaseCheck = Super::CheckState(Params))
+	{
+		return BaseCheck;
+	}
+
+	const FUniqueNetIdRef LocalAccountId = Services.Get<FAuthOSSAdapter>()->GetUniqueNetId(Params.LocalAccountId).ToSharedRef();
+	if (!LocalAccountId->IsValid())
+	{
+		UE_LOG(LogTemp, Verbose, TEXT("[FSessionsOSSAdapter::FindSessionsImpl] UniqueId not found for LocalAccountId %s"), *ToLogString(Params.LocalAccountId));
+		
+		return TOptional<FOnlineError>(Errors::InvalidUser());
+	}
+
+	if (PendingV1SessionSearchesPerUser.Contains(Params.LocalAccountId))
+	{
+		UE_LOG(LogTemp, Verbose, TEXT("[FSessionsOSSAdapter::FindSessionsImpl] Session search already in progress for LocalAccountId %s"), *ToLogString(Params.LocalAccountId));
+		
+		return TOptional<FOnlineError>(Errors::AlreadyPending());
+	}
+
+	return TOptional<FOnlineError>();
+}
+
 TFuture<TOnlineResult<FStartMatchmaking>> FSessionsOSSAdapter::StartMatchmakingImpl(const FStartMatchmaking::Params& Params)
 {
 	TPromise<TOnlineResult<FStartMatchmaking>> Promise;
@@ -706,14 +720,6 @@ TFuture<TOnlineResult<FJoinSession>> FSessionsOSSAdapter::JoinSessionImpl(const
 	TPromise<TOnlineResult<FJoinSession>> Promise;
 	TFuture<TOnlineResult<FJoinSession>> Future = Promise.GetFuture();
 
-	const FUniqueNetIdRef LocalAccountId = Services.Get<FAuthOSSAdapter>()->GetUniqueNetId(Params.LocalAccountId).ToSharedRef();
-	if (!LocalAccountId->IsValid())
-	{
-		UE_LOG(LogTemp, Verbose, TEXT("[FSessionsOSSAdapter::JoinSessionImpl] UniqueId not found for LocalAccountId %s"), *ToLogString(Params.LocalAccountId));
-		Promise.EmplaceValue(Errors::InvalidUser());
-		return Future;
-	}
-
 	MakeMulticastAdapter(this, SessionsInterface->OnJoinSessionCompleteDelegates,
 	[this, Promise = MoveTemp(Promise), Params](FName SessionName, EOnJoinSessionCompleteResult::Type Result) mutable
 	{
@@ -759,6 +765,10 @@ TFuture<TOnlineResult<FJoinSession>> FSessionsOSSAdapter::JoinSessionImpl(const
 
 	const TSharedRef<const ISession>& FoundSession = GetSessionByIdResult.GetOkValue().Session;
 
+	const FUniqueNetIdPtr LocalAccountIdPtr = Services.Get<FAuthOSSAdapter>()->GetUniqueNetId(Params.LocalAccountId);
+	check(LocalAccountIdPtr.IsValid());
+	const FUniqueNetIdRef LocalAccountId = LocalAccountIdPtr.ToSharedRef();
+
 	FOnlineSessionSearchResult SearchResult;
 
 	if (const FCustomSessionSetting* PingInMs = FoundSession->GetSessionSettings().CustomSettings.Find(OSS_ADAPTER_SESSIONS_PING_IN_MS))
@@ -773,6 +783,24 @@ TFuture<TOnlineResult<FJoinSession>> FSessionsOSSAdapter::JoinSessionImpl(const
 	return Future;
 }
 
+TOptional<FOnlineError> FSessionsOSSAdapter::CheckState(const FJoinSession::Params& Params) const
+{
+	if (TOptional<FOnlineError> BaseCheck = Super::CheckState(Params))
+	{
+		return BaseCheck;
+	}
+
+	const FUniqueNetIdRef LocalAccountId = Services.Get<FAuthOSSAdapter>()->GetUniqueNetId(Params.LocalAccountId).ToSharedRef();
+	if (!LocalAccountId->IsValid())
+	{
+		UE_LOG(LogTemp, Verbose, TEXT("[%s] UniqueId not found for LocalAccountId %s"), UTF8_TO_TCHAR(__FUNCTION__), *ToLogString(Params.LocalAccountId));
+		
+		return TOptional<FOnlineError>(Errors::InvalidUser());
+	}
+
+	return TOptional<FOnlineError>();
+}
+
 TFuture<TOnlineResult<FAddSessionMember>> FSessionsOSSAdapter::AddSessionMemberImpl(const FAddSessionMember::Params& Params)
 {
 	TPromise<TOnlineResult<FAddSessionMember>> Promise;
@@ -842,13 +870,9 @@ TFuture<TOnlineResult<FSendSessionInvite>> FSessionsOSSAdapter::SendSessionInvit
 
 	FAuthOSSAdapter* Auth = Services.Get<FAuthOSSAdapter>();
 
-	const FUniqueNetIdRef LocalAccountId = Auth->GetUniqueNetId(Params.LocalAccountId).ToSharedRef();
-	if (!LocalAccountId->IsValid())
-	{
-		UE_LOG(LogTemp, Verbose, TEXT("[FSessionsOSSAdapter::SendSessionInviteImpl] UniqueId not found for LocalAccountId %s"), *ToLogString(Params.LocalAccountId));
-		Promise.EmplaceValue(Errors::InvalidUser());
-		return Future;
-	}
+	const FUniqueNetIdPtr LocalAccountIdPtr = Services.Get<FAuthOSSAdapter>()->GetUniqueNetId(Params.LocalAccountId);
+	check(LocalAccountIdPtr.IsValid());
+	const FUniqueNetIdRef LocalAccountId = LocalAccountIdPtr.ToSharedRef();
 
 	TArray<FUniqueNetIdRef> TargetUserNetIds;
 	for (const FAccountId& TargetUser : Params.TargetUsers)
@@ -867,6 +891,24 @@ TFuture<TOnlineResult<FSendSessionInvite>> FSessionsOSSAdapter::SendSessionInvit
 	return Future;
 }
 
+TOptional<FOnlineError> FSessionsOSSAdapter::CheckState(const FSendSessionInvite::Params& Params) const
+{
+	if (TOptional<FOnlineError> BaseCheck = Super::CheckState(Params))
+	{
+		return BaseCheck;
+	}
+
+	const FUniqueNetIdRef LocalAccountId = Services.Get<FAuthOSSAdapter>()->GetUniqueNetId(Params.LocalAccountId).ToSharedRef();
+	if (!LocalAccountId->IsValid())
+	{
+		UE_LOG(LogTemp, Verbose, TEXT("[%s] UniqueId not found for LocalAccountId %s"), UTF8_TO_TCHAR(__FUNCTION__), *ToLogString(Params.LocalAccountId));
+
+		return TOptional<FOnlineError>(Errors::InvalidUser());
+	}
+
+	return TOptional<FOnlineError>();
+}
+
 TOnlineResult<FGetResolvedConnectString> FSessionsOSSAdapter::GetResolvedConnectString(const FGetResolvedConnectString::Params& Params)
 {
 	if (Params.SessionId.IsValid())
diff --git a/Engine/Plugins/Online/OnlineServicesOSSAdapter/Source/Public/Online/SessionsOSSAdapter.h b/Engine/Plugins/Online/OnlineServicesOSSAdapter/Source/Public/Online/SessionsOSSAdapter.h
--- a/Engine/Plugins/Online/OnlineServicesOSSAdapter/Source/Public/Online/SessionsOSSAdapter.h
+++ b/Engine/Plugins/Online/OnlineServicesOSSAdapter/Source/Public/Online/SessionsOSSAdapter.h
@@ -77,6 +77,11 @@ public:
 
 	TOnlineResult<FGetResolvedConnectString> GetResolvedConnectString(const FGetResolvedConnectString::Params& Params);
 
+protected:
+	virtual TOptional<FOnlineError> CheckState(const FFindSessions::Params& Params) const override;
+	virtual TOptional<FOnlineError> CheckState(const FJoinSession::Params& Params) const override;
+	virtual TOptional<FOnlineError> CheckState(const FSendSessionInvite::Params& Params) const override;
+
 private:
 	/** Updates a local session's state by retrieving the equivalent from the OSS */
 	bool UpdateV2Session(const FName& SessionName);
