Changed misleading variable names

diff --git a/Engine/Plugins/Online/OnlineServicesOSSAdapter/Source/Private/Online/SessionsOSSAdapter.cpp b/Engine/Plugins/Online/OnlineServicesOSSAdapter/Source/Private/Online/SessionsOSSAdapter.cpp
--- a/Engine/Plugins/Online/OnlineServicesOSSAdapter/Source/Private/Online/SessionsOSSAdapter.cpp
+++ b/Engine/Plugins/Online/OnlineServicesOSSAdapter/Source/Private/Online/SessionsOSSAdapter.cpp
@@ -981,20 +981,20 @@ FOnlineSessionSettings FSessionsOSSAdapter::BuildV1SettingsForCreate(const FCrea
 	{
 		Result.BuildUniqueId = BuildUniqueId->Data.GetInt64();
 	}
-	if (const FCustomSessionSetting* BuildUniqueId = Params.SessionSettings.CustomSettings.Find(OSS_ADAPTER_SESSIONS_USE_LOBBIES_IF_AVAILABLE))
+	if (const FCustomSessionSetting* UseLobbiesIfAvailable = Params.SessionSettings.CustomSettings.Find(OSS_ADAPTER_SESSIONS_USE_LOBBIES_IF_AVAILABLE))
 	{
-		Result.bUseLobbiesIfAvailable = BuildUniqueId->Data.GetBoolean();
+		Result.bUseLobbiesIfAvailable = UseLobbiesIfAvailable->Data.GetBoolean();
 	}
-	if (const FCustomSessionSetting* BuildUniqueId = Params.SessionSettings.CustomSettings.Find(OSS_ADAPTER_SESSIONS_USE_LOBBIES_VOICE_CHAT_IF_AVAILABLE))
+	if (const FCustomSessionSetting* UseLobbiesVoiceChatIfAvailable = Params.SessionSettings.CustomSettings.Find(OSS_ADAPTER_SESSIONS_USE_LOBBIES_VOICE_CHAT_IF_AVAILABLE))
 	{
-		Result.bUseLobbiesVoiceChatIfAvailable = BuildUniqueId->Data.GetBoolean();
+		Result.bUseLobbiesVoiceChatIfAvailable = UseLobbiesVoiceChatIfAvailable->Data.GetBoolean();
 	}
 
 	Result.bUsesPresence = Params.bPresenceEnabled;
 
-	if (const FCustomSessionSetting* BuildUniqueId = Params.SessionSettings.CustomSettings.Find(OSS_ADAPTER_SESSIONS_USES_STATS))
+	if (const FCustomSessionSetting* UsesStats = Params.SessionSettings.CustomSettings.Find(OSS_ADAPTER_SESSIONS_USES_STATS))
 	{
-		Result.bUsesStats = BuildUniqueId->Data.GetBoolean();
+		Result.bUsesStats = UsesStats->Data.GetBoolean();
 	}
 	Result.NumPrivateConnections = Params.SessionSettings.NumMaxConnections;
 	Result.NumPublicConnections = Params.SessionSettings.NumMaxConnections;
