- Add OnlineServicesMcp module

diff --git a/Engine/Plugins/Online/OnlineServicesOSSAdapter/Source/OnlineServicesOSSAdapter.Build.cs b/Engine/Plugins/Online/OnlineServicesOSSAdapter/Source/OnlineServicesOSSAdapter.Build.cs
--- a/Engine/Plugins/Online/OnlineServicesOSSAdapter/Source/OnlineServicesOSSAdapter.Build.cs
+++ b/Engine/Plugins/Online/OnlineServicesOSSAdapter/Source/OnlineServicesOSSAdapter.Build.cs
@@ -11,6 +11,7 @@ public class OnlineServicesOSSAdapter : ModuleRules
 			new string[] {
 				"Core",
 				"OnlineServicesInterface",
+				"OnlineServicesCommon",
 				"OnlineSubsystem",
 				"Json"
 			}
@@ -19,7 +20,6 @@ public class OnlineServicesOSSAdapter : ModuleRules
 		PrivateDependencyModuleNames.AddRange(
 			new string[] {
 				"CoreOnline",
-				"OnlineServicesCommon"
 			}
 		);
 	}
diff --git a/Engine/Plugins/Online/OnlineServicesOSSAdapter/Source/Private/Online/AuthOSSAdapter.cpp b/Engine/Plugins/Online/OnlineServicesOSSAdapter/Source/Private/Online/AuthOSSAdapter.cpp
--- a/Engine/Plugins/Online/OnlineServicesOSSAdapter/Source/Private/Online/AuthOSSAdapter.cpp
+++ b/Engine/Plugins/Online/OnlineServicesOSSAdapter/Source/Private/Online/AuthOSSAdapter.cpp
@@ -521,11 +521,15 @@ TOnlineAsyncOpHandle<FAuthQueryExternalAuthToken> FAuthOSSAdapter::QueryExternal
 
 FUniqueNetIdPtr FAuthOSSAdapter::GetUniqueNetId(FAccountId AccountId) const
 {
+	check(AccountId.GetOnlineServicesType() == GetOnlineServicesOSSAdapter().GetServicesProvider());
+
 	return GetOnlineServicesOSSAdapter().GetAccountIdRegistry().GetIdValue(AccountId);
 }
 
 FAccountId FAuthOSSAdapter::GetAccountId(const FUniqueNetIdRef& UniqueNetId) const
 {
+	check(UniqueNetId->GetType() == GetSubsystem().GetSubsystemName());
+	
 	return GetOnlineServicesOSSAdapter().GetAccountIdRegistry().FindOrAddHandle(UniqueNetId);
 }
 
diff --git a/Engine/Plugins/Online/OnlineServicesOSSAdapter/Source/Public/Online/AuthOSSAdapter.h b/Engine/Plugins/Online/OnlineServicesOSSAdapter/Source/Public/Online/AuthOSSAdapter.h
--- a/Engine/Plugins/Online/OnlineServicesOSSAdapter/Source/Public/Online/AuthOSSAdapter.h
+++ b/Engine/Plugins/Online/OnlineServicesOSSAdapter/Source/Public/Online/AuthOSSAdapter.h
@@ -50,7 +50,7 @@ public:
 	void Unregister(FAccountId AccountId);
 };
 
-class FAuthOSSAdapter : public FAuthCommon
+class ONLINESERVICESOSSADAPTER_API FAuthOSSAdapter : public FAuthCommon
 {
 public:
 	using Super = FAuthCommon;
diff --git a/Engine/Plugins/Online/OnlineServicesOSSAdapter/Source/Public/Online/OnlineServicesOSSAdapter.h b/Engine/Plugins/Online/OnlineServicesOSSAdapter/Source/Public/Online/OnlineServicesOSSAdapter.h
--- a/Engine/Plugins/Online/OnlineServicesOSSAdapter/Source/Public/Online/OnlineServicesOSSAdapter.h
+++ b/Engine/Plugins/Online/OnlineServicesOSSAdapter/Source/Public/Online/OnlineServicesOSSAdapter.h
@@ -10,15 +10,15 @@ class IOnlineSubsystem;
 
 namespace UE::Online {
 
-class FOnlineServicesOSSAdapter : public FOnlineServicesCommon
+class ONLINESERVICESOSSADAPTER_API FOnlineServicesOSSAdapter : public FOnlineServicesCommon
 {
 public:
 	using Super = FOnlineServicesCommon;
 
-	ONLINESERVICESOSSADAPTER_API FOnlineServicesOSSAdapter(EOnlineServices InServicesType, const FString& InConfigName, FName InInstanceName, IOnlineSubsystem* InSubsystem);
+	FOnlineServicesOSSAdapter(EOnlineServices InServicesType, const FString& InConfigName, FName InInstanceName, IOnlineSubsystem* InSubsystem);
 
-	ONLINESERVICESOSSADAPTER_API virtual void RegisterComponents() override;
-	ONLINESERVICESOSSADAPTER_API virtual void Initialize() override;
+	virtual void RegisterComponents() override;
+	virtual void Initialize() override;
 	virtual EOnlineServices GetServicesProvider() const override { return ServicesType; }
 	virtual TOnlineResult<FGetResolvedConnectString> GetResolvedConnectString(FGetResolvedConnectString::Params&& Params) override;
 
