FOnlineError support for string error codes.
Fixed issue with content resolver content service provider randomly cleaning up results.

diff --git a/Engine/Plugins/Online/OnlineServices/Source/OnlineServicesInterface/Private/Online/OnlineError.cpp b/Engine/Plugins/Online/OnlineServices/Source/OnlineServicesInterface/Private/Online/OnlineError.cpp
--- a/Engine/Plugins/Online/OnlineServices/Source/OnlineServicesInterface/Private/Online/OnlineError.cpp
+++ b/Engine/Plugins/Online/OnlineServices/Source/OnlineServicesInterface/Private/Online/OnlineError.cpp
@@ -88,6 +88,7 @@ namespace UE::Online{
 			FString ErrorCodeString = OnlineError.GetErrorId();
 			ErrorMessageArgs.Add(TEXT("ErrorCode"), FText::FromString(ErrorCodeString));
 			Writer.AddString(ANSITEXTVIEW("ErrorCode"), ErrorCodeString);
+			
 		}
 
 		if (OnlineError.Details)
@@ -96,6 +97,8 @@ namespace UE::Online{
 			ErrorMessageArgs.Add(TEXT("ErrorDetails"), ErrorDetails);
 			Writer.AddString(ANSITEXTVIEW("ErrorDetails"), ErrorDetails.ToString());
 			bHasDetails = true;
+			ErrorMessageArgs.Add(TEXT("FriendlyErrorCode"), FText::FromString(OnlineError.Details->GetFriendlyErrorCode(OnlineError)));
+			Writer.AddString(ANSITEXTVIEW("FriendlyErrorCode"), OnlineError.Details->GetFriendlyErrorCode(OnlineError));
 		}
 		
 		if (OnlineError.GetInner() != nullptr)
@@ -108,11 +111,11 @@ namespace UE::Online{
 
 		if (bHasInner && bHasDetails)
 		{
-			Writer.AddString("$text", FText::Format(NSLOCTEXT("OnlineError", "ErrorDetailsErrorInner", "{ErrorCode}, {ErrorDetails}, InnerError:{InnerError}"), ErrorMessageArgs).ToString());
+			Writer.AddString("$text", FText::Format(NSLOCTEXT("OnlineError", "ErrorDetailsErrorInner", "{ErrorCode} ({FriendlyErrorCode}), {ErrorDetails}, InnerError:{InnerError}"), ErrorMessageArgs).ToString());
 		}
 		else if (bHasDetails)
 		{
-			Writer.AddString("$text", FText::Format(NSLOCTEXT("OnlineError", "ErrorDetails", "{ErrorCode}, {ErrorDetails}"), ErrorMessageArgs).ToString());
+			Writer.AddString("$text", FText::Format(NSLOCTEXT("OnlineError", "ErrorDetails", "{ErrorCode} ({FriendlyErrorCode}), {ErrorDetails}"), ErrorMessageArgs).ToString());
 		}
 		else if (bHasInner)
 		{
diff --git a/Engine/Plugins/Online/OnlineServices/Source/OnlineServicesInterface/Public/Online/OnlineError.h b/Engine/Plugins/Online/OnlineServices/Source/OnlineServicesInterface/Public/Online/OnlineError.h
--- a/Engine/Plugins/Online/OnlineServices/Source/OnlineServicesInterface/Public/Online/OnlineError.h
+++ b/Engine/Plugins/Online/OnlineServices/Source/OnlineServicesInterface/Public/Online/OnlineError.h
@@ -14,6 +14,7 @@ class FOnlineError;
 class IOnlineErrorDetails
 {
 public:
+	virtual FString GetFriendlyErrorCode(const FOnlineError&) const = 0;
 	virtual FText GetText(const FOnlineError&) const = 0;
 	virtual FString GetLogString(const FOnlineError&) const = 0;
 };
@@ -21,14 +22,21 @@ public:
 class FOnlineErrorDetails : public IOnlineErrorDetails
 {
 public:
-	FOnlineErrorDetails(FString&& InLogString, FText&& InText)
-		: LogString(MoveTemp(InLogString))
+	FOnlineErrorDetails(FString&& InFriendlyErrorCode, FString&& InLogString, FText&& InText) 
+		: FriendlyErrorCode(MoveTemp(InFriendlyErrorCode))
+		, LogString(MoveTemp(InLogString))
 		, Text(MoveTemp(InText))
 	{
 	}
 
+
 	virtual ~FOnlineErrorDetails() {}
 
+	virtual FString GetFriendlyErrorCode(const FOnlineError&) const override
+	{
+		return FriendlyErrorCode;
+	}
+
 	virtual FString GetLogString(const FOnlineError&) const override
 	{
 		return LogString;
@@ -40,6 +48,7 @@ public:
 	}
 
 protected:
+	FString FriendlyErrorCode;
 	FString LogString;
 	FText Text;
 };
diff --git a/Engine/Plugins/Online/OnlineServices/Source/OnlineServicesInterface/Public/Online/OnlineErrorCode.h b/Engine/Plugins/Online/OnlineServices/Source/OnlineServicesInterface/Public/Online/OnlineErrorCode.h
--- a/Engine/Plugins/Online/OnlineServices/Source/OnlineServicesInterface/Public/Online/OnlineErrorCode.h
+++ b/Engine/Plugins/Online/OnlineServices/Source/OnlineServicesInterface/Public/Online/OnlineErrorCode.h
@@ -79,7 +79,7 @@ namespace ErrorCode { namespace Category { \
 			{ \
 				InnerPtr = MakeShared<UE::Online::FOnlineError, ESPMode::ThreadSafe>(Inner.GetValue()); \
 			} \
-			return UE::Online::FOnlineError( ErrorCode::Create(UE::Online::Errors::ErrorCode::Category::CategoryName##_System, UE::Online::Errors::ErrorCode::Category::CategoryName, ErrorCodeValue), MakeShared<UE::Online::FOnlineErrorDetails, ESPMode::ThreadSafe>(ErrorMessage, ErrorText), InnerPtr); \
+			return UE::Online::FOnlineError( ErrorCode::Create(UE::Online::Errors::ErrorCode::Category::CategoryName##_System, UE::Online::Errors::ErrorCode::Category::CategoryName, ErrorCodeValue), MakeShared<UE::Online::FOnlineErrorDetails, ESPMode::ThreadSafe>(FString(#CategoryName"."#Name), ErrorMessage, ErrorText), InnerPtr); \
 		}
 
 #define UE_ONLINE_ERROR(CategoryName, Name, ErrorCodeValue, ErrorMessage, ErrorText) \
diff --git a/Engine/Plugins/Online/OnlineServicesEOSGS/Source/Private/Online/OnlineErrorEOSGS.cpp b/Engine/Plugins/Online/OnlineServicesEOSGS/Source/Private/Online/OnlineErrorEOSGS.cpp
--- a/Engine/Plugins/Online/OnlineServicesEOSGS/Source/Private/Online/OnlineErrorEOSGS.cpp
+++ b/Engine/Plugins/Online/OnlineServicesEOSGS/Source/Private/Online/OnlineErrorEOSGS.cpp
@@ -12,6 +12,8 @@ class FOnlineErrorDetailsEOS : public IOnlineErrorDetails
 public:
 	static const TSharedRef<const IOnlineErrorDetails, ESPMode::ThreadSafe>& Get();
 
+	virtual FString GetFriendlyErrorCode(const FOnlineError& OnlineError) const override;
+
 	virtual FText GetText(const FOnlineError& OnlineError) const override;
 	virtual FString GetLogString(const FOnlineError& OnlineError) const override;
 };
@@ -22,6 +24,12 @@ const TSharedRef<const IOnlineErrorDetails, ESPMode::ThreadSafe>& FOnlineErrorDe
 	return Instance;
 }
 
+FString FOnlineErrorDetailsEOS::GetFriendlyErrorCode(const FOnlineError& OnlineError) const
+{
+	const EOS_EResult EosResult = EOS_EResult(OnlineError.GetValue());
+	return LexToString(EosResult);
+}
+
 FText FOnlineErrorDetailsEOS::GetText(const FOnlineError& OnlineError) const
 {
 	const EOS_EResult EosResult = EOS_EResult(OnlineError.GetValue());
diff --git a/Engine/Plugins/Online/OnlineServicesOSSAdapter/Source/Private/Online/ErrorsOSSAdapter.cpp b/Engine/Plugins/Online/OnlineServicesOSSAdapter/Source/Private/Online/ErrorsOSSAdapter.cpp
--- a/Engine/Plugins/Online/OnlineServicesOSSAdapter/Source/Private/Online/ErrorsOSSAdapter.cpp
+++ b/Engine/Plugins/Online/OnlineServicesOSSAdapter/Source/Private/Online/ErrorsOSSAdapter.cpp
@@ -70,8 +70,9 @@ ONLINESERVICESOSSADAPTER_API inline FOnlineError FromOssError(const FOnlineError
 	{
 		ErrorCode = Result.GetErrorCode();
 	}
+	FString RawErrorCode = Result.GetErrorCode();
 
-	return Internal_OssWrapInner(FOnlineError(ErrorCode::Create(ErrorCode::Category::Oss_System, ErrorCode::Category::Oss, (uint32)Result.GetErrorResult()), MakeShared<FOnlineErrorDetails, ESPMode::ThreadSafe>(MoveTemp(ErrorCode), MoveTemp(ErrorMessage)), nullptr), Result);
+	return Internal_OssWrapInner(FOnlineError(ErrorCode::Create(ErrorCode::Category::Oss_System, ErrorCode::Category::Oss, (uint32)Result.GetErrorResult()), MakeShared<FOnlineErrorDetails, ESPMode::ThreadSafe>(MoveTemp(RawErrorCode), MoveTemp(ErrorCode), MoveTemp(ErrorMessage)), nullptr), Result);
 }
 
 FOnlineError FromOssErrorCode(const FString& ErrorCode)
@@ -149,7 +150,8 @@ FOnlineError FromOssErrorCode(const FString& ErrorCode)
 	// construct an OSS FOnlineError from the result
 	::FOnlineError Result = ::FOnlineError::CreateError(FString(), ErrorResult, ErrorCode);
 
-	return Internal_OssWrapInner(FOnlineError(ErrorCode::Create(ErrorCode::Category::Oss_System, ErrorCode::Category::Oss, (uint32)Result.GetErrorResult()), MakeShared<FOnlineErrorDetails, ESPMode::ThreadSafe>(CopyTemp(ErrorCode), CopyTemp(Result.GetErrorMessage())), nullptr), Result);
+	return Internal_OssWrapInner(FOnlineError(ErrorCode::Create(ErrorCode::Category::Oss_System, ErrorCode::Category::Oss, (uint32)Result.GetErrorResult()), 
+		MakeShared<FOnlineErrorDetails, ESPMode::ThreadSafe>(CopyTemp(ErrorCode), CopyTemp(ErrorCode), CopyTemp(Result.GetErrorMessage())), nullptr), Result);
 }
 
 } //namespace UE::Online::Errors
