- Improve documentation of DelegateAdapter

#skipci
#rnx

diff --git a/Engine/Plugins/Online/OnlineServicesOSSAdapter/Source/Public/Online/DelegateAdapter.h b/Engine/Plugins/Online/OnlineServicesOSSAdapter/Source/Public/Online/DelegateAdapter.h
--- a/Engine/Plugins/Online/OnlineServicesOSSAdapter/Source/Public/Online/DelegateAdapter.h
+++ b/Engine/Plugins/Online/OnlineServicesOSSAdapter/Source/Public/Online/DelegateAdapter.h
@@ -7,15 +7,20 @@
 
 /*
 	BASIC USAGE
-		The Delegate Adapter can be used in place of an FOn{Name}Delegate::CreateLambda call
-
-		i.e.
+		The Delegate Adapter can be used in place of an FOn{Name}Delegate::CreateLambda call:
 			
 			Identity->GetAuthToken(LocalUserNum, FOnAuthTokenGetComplete::CreateLambda[this](int32 UserNum, FString LoginToken){ ... });
 			->
 			Identity->GetAuthToken(LocalUserNum, *MakeDelegateAdapter(this, [this](int32 UserNum, FString LoginToken){ ... });
 
-		The second instance, in addition to being slightly shorter and easier to remember, will be a unique ptr (can MoveTemp into the closure) and also will automatically do a weak ptr validity check on this
+		The second instance, in addition to being slightly shorter and easier to remember, will be a unique ptr (can MoveTemp into the closure) and also will automatically do a weak ptr validity check on "this"
+
+		The Delegate Adapter can also be used to bind to TDelegates:
+
+			TDelegate<void(bool)> Delegate;
+			Delegate = *MakeDelegateAdapter(this, [this](bool bParam) {});
+
+		This is useful when you need to capture something non-copyable, e.g. a TPromise, as TDelegate's need to be copyable.
 
 	The lifetime of the adapter is a shared ptr attached to the input delegate and will live as long as that delegate is bound.
 */
