Add FAuthOSSAdapterPtr alias

diff --git a/Engine/Plugins/Online/OnlineServicesOSSAdapter/Source/Public/Online/AuthOSSAdapter.h b/Engine/Plugins/Online/OnlineServicesOSSAdapter/Source/Public/Online/AuthOSSAdapter.h
--- a/Engine/Plugins/Online/OnlineServicesOSSAdapter/Source/Public/Online/AuthOSSAdapter.h
+++ b/Engine/Plugins/Online/OnlineServicesOSSAdapter/Source/Public/Online/AuthOSSAdapter.h
@@ -50,6 +50,8 @@ public:
 	void Unregister(FAccountId AccountId);
 };
 
+using FAuthOSSAdapterPtr = TSharedPtr<class FAuthOSSAdapter>;
+
 class ONLINESERVICESOSSADAPTER_API FAuthOSSAdapter : public FAuthCommon
 {
 public:
