- Edigrate GetFriendlyErrorCode

diff --git a/Engine/Plugins/Online/OnlineServices/Source/OnlineServicesInterface/Public/Online/OnlineError.h b/Engine/Plugins/Online/OnlineServices/Source/OnlineServicesInterface/Public/Online/OnlineError.h
--- a/Engine/Plugins/Online/OnlineServices/Source/OnlineServicesInterface/Public/Online/OnlineError.h
+++ b/Engine/Plugins/Online/OnlineServices/Source/OnlineServicesInterface/Public/Online/OnlineError.h
@@ -145,6 +145,22 @@ public:
 		return ErrorId;
 	}
 
+	FString GetFriendlyErrorCode() const
+	{
+		if (ErrorCode == Errors::ErrorCode::Success)
+		{
+			return TEXT("Success");
+		}
+		else if (Details)
+		{
+			return Details->GetFriendlyErrorCode(*this);
+		}
+		else
+		{
+			return GetErrorId();
+		}
+	}
+
 	const FOnlineError* GetInner() const
 	{
 		return Inner.IsValid() ? Inner.Get() : nullptr;
diff --git a/Engine/Plugins/Online/OnlineServices/Source/OnlineServicesInterface/Public/Online/OnlineErrorCode.h b/Engine/Plugins/Online/OnlineServices/Source/OnlineServicesInterface/Public/Online/OnlineErrorCode.h
--- a/Engine/Plugins/Online/OnlineServices/Source/OnlineServicesInterface/Public/Online/OnlineErrorCode.h
+++ b/Engine/Plugins/Online/OnlineServices/Source/OnlineServicesInterface/Public/Online/OnlineErrorCode.h
@@ -78,7 +78,7 @@ namespace ErrorCode { namespace Category { \
 			{ \
 				InnerPtr = MakeShared<UE::Online::FOnlineError, ESPMode::ThreadSafe>(Inner.GetValue()); \
 			} \
-			return UE::Online::FOnlineError( UE::Online::Errors::ErrorCode::Create(ErrorCode::Category::CategoryName##_System, ErrorCode::Category::CategoryName, ErrorCodeValue), MakeShared<UE::Online::FOnlineErrorDetails, ESPMode::ThreadSafe>(FString(#CategoryName"."#Name), ErrorMessage, ErrorText), InnerPtr); \
+			return UE::Online::FOnlineError( UE::Online::Errors::ErrorCode::Create(UE::Online::Errors::ErrorCode::Category::CategoryName##_System, UE::Online::Errors::ErrorCode::Category::CategoryName, ErrorCodeValue), MakeShared<UE::Online::FOnlineErrorDetails, ESPMode::ThreadSafe>(FString(#CategoryName"."#Name), ErrorMessage, ErrorText), InnerPtr); \
 		}
 
 /** Macro to define an error within a category. Must be used within the UE::Online::Errors namespace. The error will be accessible as UE::Online::Errors::CategoryName::Name() */
diff --git a/Engine/Plugins/Online/OnlineServicesOSSAdapter/Source/Private/Online/ErrorsOSSAdapter.cpp b/Engine/Plugins/Online/OnlineServicesOSSAdapter/Source/Private/Online/ErrorsOSSAdapter.cpp
--- a/Engine/Plugins/Online/OnlineServicesOSSAdapter/Source/Private/Online/ErrorsOSSAdapter.cpp
+++ b/Engine/Plugins/Online/OnlineServicesOSSAdapter/Source/Private/Online/ErrorsOSSAdapter.cpp
@@ -57,7 +57,7 @@ inline FOnlineError Internal_OssWrapInner(FOnlineError Error, const FOnlineError
 	return Internal_OssGetCommonError(Error, Result.GetErrorResult());
 }
 
-ONLINESERVICESOSSADAPTER_API inline FOnlineError FromOssError(const FOnlineErrorOss& Result)
+FOnlineError FromOssError(const FOnlineErrorOss& Result)
 {
 	FString ErrorCode;
 	FText ErrorMessage = Result.GetErrorMessage();
@@ -70,9 +70,10 @@ ONLINESERVICESOSSADAPTER_API inline FOnlineError FromOssError(const FOnlineError
 	{
 		ErrorCode = Result.GetErrorCode();
 	}
-	FString RawErrorCode = Result.GetErrorCode();
+	
+	FString FriendlyErrorCode = Result.GetErrorCode();
 
-	return Internal_OssWrapInner(FOnlineError(ErrorCode::Create(ErrorCode::Category::Oss_System, ErrorCode::Category::Oss, (uint32)Result.GetErrorResult()), MakeShared<FOnlineErrorDetails, ESPMode::ThreadSafe>(MoveTemp(RawErrorCode), MoveTemp(ErrorCode), MoveTemp(ErrorMessage)), nullptr), Result);
+	return Internal_OssWrapInner(FOnlineError(ErrorCode::Create(ErrorCode::Category::Oss_System, ErrorCode::Category::Oss, (uint32)Result.GetErrorResult()), MakeShared<FOnlineErrorDetails, ESPMode::ThreadSafe>(MoveTemp(FriendlyErrorCode), MoveTemp(ErrorCode), MoveTemp(ErrorMessage)), nullptr), Result);
 }
 
 FOnlineError FromOssErrorCode(const FString& ErrorCode)
diff --git a/Engine/Plugins/Online/OnlineServicesOSSAdapter/Source/Public/Online/ErrorsOSSAdapter.h b/Engine/Plugins/Online/OnlineServicesOSSAdapter/Source/Public/Online/ErrorsOSSAdapter.h
--- a/Engine/Plugins/Online/OnlineServicesOSSAdapter/Source/Public/Online/ErrorsOSSAdapter.h
+++ b/Engine/Plugins/Online/OnlineServicesOSSAdapter/Source/Public/Online/ErrorsOSSAdapter.h
@@ -9,7 +9,7 @@ namespace UE::Online::Errors {
 
 	UE_ONLINE_ERROR_CATEGORY(Oss, Engine, 0x7, "Oss")
 	ONLINESERVICESOSSADAPTER_API FOnlineError FromOssError(const FOnlineErrorOss& Result);
-	FOnlineError FromOssErrorCode(const FString& ErrorCode);
+	ONLINESERVICESOSSADAPTER_API FOnlineError FromOssErrorCode(const FString& ErrorCode);
 	
 	// if you would like to add a parser that transforms custom platform errors into readable online service types, you can call this in your platform's error handler
 	using TOssPlatformErrorHandler = TUniqueFunction<TOptional<FOnlineError>(const FOnlineError&, const FOnlineErrorOss&)>;
