Query users on OSS adapter init.

diff --git a/Engine/Plugins/Online/OnlineServicesOSSAdapter/Source/Private/Online/AuthOSSAdapter.cpp b/Engine/Plugins/Online/OnlineServicesOSSAdapter/Source/Private/Online/AuthOSSAdapter.cpp
--- a/Engine/Plugins/Online/OnlineServicesOSSAdapter/Source/Private/Online/AuthOSSAdapter.cpp
+++ b/Engine/Plugins/Online/OnlineServicesOSSAdapter/Source/Private/Online/AuthOSSAdapter.cpp
@@ -144,6 +144,13 @@ void FAuthOSSAdapter::PostInitialize()
 							TranslateLoginStatus(NewStatus) });
 					}
 				});
+			if (FUniqueNetIdPtr NetId = Identity->GetUniquePlayerId(LocalPlayerNum))
+			{
+				HandleLoginStatusChangedImplOp(FAuthHandleLoginStatusChangedImpl::Params{
+								FPlatformMisc::GetPlatformUserForUserIndex(LocalPlayerNum),
+								static_cast<FOnlineServicesOSSAdapter&>(Services).GetAccountIdRegistry().FindOrAddHandle(NetId.ToSharedRef()),
+								TranslateLoginStatus(Identity->GetLoginStatus(LocalPlayerNum)) });
+			}			
 		}
 	}
 }
