Fix for Sessions adapter not working properly for creating/destroying. Also left in the added testing debug.

diff --git a/Engine/Plugins/Online/OnlineServicesOSSAdapter/Source/Private/Online/SessionsOSSAdapter.cpp b/Engine/Plugins/Online/OnlineServicesOSSAdapter/Source/Private/Online/SessionsOSSAdapter.cpp
--- a/Engine/Plugins/Online/OnlineServicesOSSAdapter/Source/Private/Online/SessionsOSSAdapter.cpp
+++ b/Engine/Plugins/Online/OnlineServicesOSSAdapter/Source/Private/Online/SessionsOSSAdapter.cpp
@@ -12,6 +12,7 @@
 #include "Online/UserInfoOSSAdapter.h"
 #include "Online/DelegateAdapter.h"
 #include "Online/ErrorsOSSAdapter.h"
+#include "Online/OnlineSessionNames.h"
 
 #include "Algo/Find.h"
 
@@ -973,7 +974,6 @@ FOnlineSessionSettings FSessionsOSSAdapter::BuildV1SettingsForCreate(const FCrea
 	Result.bAllowJoinInProgress = Params.SessionSettings.bAllowNewMembers;
 	Result.bAllowJoinViaPresence = Params.SessionSettings.bAllowNewMembers;
 	Result.bAllowJoinViaPresenceFriendsOnly = Params.SessionSettings.bAllowNewMembers;
-	Result.Settings.Add(OSS_ADAPTER_SESSIONS_ALLOW_SANCTIONED_PLAYERS, Params.bAllowSanctionedPlayers);
 	Result.bAntiCheatProtected = Params.bAntiCheatProtected;
 	Result.bIsDedicated = IsRunningDedicatedServer();
 	Result.bIsLANMatch = Params.bIsLANSession;
@@ -999,10 +999,12 @@ FOnlineSessionSettings FSessionsOSSAdapter::BuildV1SettingsForCreate(const FCrea
 	}
 	Result.NumPrivateConnections = Params.SessionSettings.NumMaxConnections;
 	Result.NumPublicConnections = Params.SessionSettings.NumMaxConnections;
-	Result.Settings.Add(OSS_ADAPTER_SESSIONS_SCHEMA_NAME, Params.SessionSettings.SchemaName.ToString());
 	Result.SessionIdOverride = Params.SessionIdOverride;
 
 	Result.Settings = GetV1SessionSettings(Params.SessionSettings.CustomSettings);
+	Result.Settings.Add(OSS_ADAPTER_SESSIONS_SCHEMA_NAME, Params.SessionSettings.SchemaName.ToString());
+	Result.Settings.Add(SETTING_SESSION_TEMPLATE_NAME, Params.SessionSettings.SchemaName.ToString());
+	Result.Settings.Add(OSS_ADAPTER_SESSIONS_ALLOW_SANCTIONED_PLAYERS, Params.bAllowSanctionedPlayers);
 
 	FOnlineServicesOSSAdapter& ServicesOSSAdapter = static_cast<FOnlineServicesOSSAdapter&>(Services);
 
@@ -1050,10 +1052,11 @@ FOnlineSessionSettings FSessionsOSSAdapter::BuildV1SettingsForUpdate(const FAcco
 	}
 	Result.NumPrivateConnections = Session->GetSessionSettings().NumMaxConnections;
 	Result.NumPublicConnections = Session->GetSessionSettings().NumMaxConnections;
-	Result.Settings.Add(OSS_ADAPTER_SESSIONS_SCHEMA_NAME, Session->GetSessionSettings().SchemaName.ToString());
 	Result.SessionIdOverride = Session->GetSessionInfo().SessionIdOverride;
 
 	Result.Settings = GetV1SessionSettings(Session->GetSessionSettings().CustomSettings);
+	Result.Settings.Add(OSS_ADAPTER_SESSIONS_SCHEMA_NAME, Session->GetSessionSettings().SchemaName.ToString());
+	Result.Settings.Add(SETTING_SESSION_TEMPLATE_NAME, Session->GetSessionSettings().SchemaName.ToString());
 	
 	FOnlineServicesOSSAdapter& ServicesOSSAdapter = static_cast<FOnlineServicesOSSAdapter&>(Services);
 
