On login status changed delegates isn't being triggered on refresh controllers, adapter being out of sync with the underlying Identity subsystem.

- Needed to follow up with an additional change to add the user to the registry if they happen to be signed in before initilization

diff --git a/Engine/Plugins/Online/OnlineServicesOSSAdapter/Source/Private/Online/AuthOSSAdapter.cpp b/Engine/Plugins/Online/OnlineServicesOSSAdapter/Source/Private/Online/AuthOSSAdapter.cpp
--- a/Engine/Plugins/Online/OnlineServicesOSSAdapter/Source/Private/Online/AuthOSSAdapter.cpp
+++ b/Engine/Plugins/Online/OnlineServicesOSSAdapter/Source/Private/Online/AuthOSSAdapter.cpp
@@ -591,6 +591,21 @@ TOnlineAsyncOpHandle<FAuthHandleLoginStatusChangedImpl> FAuthOSSAdapter::HandleL
 	{
 		const FAuthHandleLoginStatusChangedImpl::Params& Params = InAsyncOp.GetParams();
 		TSharedPtr<FAccountInfoOSSAdapter> AccountInfoOSSAdapter = AccountInfoRegistryOSSAdapter.Find(Params.AccountId);
+		if (!AccountInfoOSSAdapter)
+		{
+			TSharedPtr<FAccountInfoOSSAdapter> NewAccountInfoOSSAdapter = MakeShared<FAccountInfoOSSAdapter>();
+			NewAccountInfoOSSAdapter->AccountId = Params.AccountId;
+			NewAccountInfoOSSAdapter->PlatformUserId = Params.PlatformUserId;
+			NewAccountInfoOSSAdapter->LocalUserNum = GetIdentityInterface()->GetLocalUserNumFromPlatformUserId(Params.PlatformUserId);
+
+			UE_LOG(LogOnlineServices, Log, TEXT("[FAuthOSSAdapter::HandleLoginStatusChangedImplOp][%s] Attempt to register [%s] account that have been initialized before the adapter has fully initialize."),
+				*GetSubsystem().GetSubsystemName().ToString(), *ToLogString(NewAccountInfoOSSAdapter->AccountId));
+
+			// Add the user to the registry if they happen to be signed in before initilization
+			AccountInfoRegistryOSSAdapter.Register(NewAccountInfoOSSAdapter.ToSharedRef());
+			AccountInfoOSSAdapter = AccountInfoRegistryOSSAdapter.Find(Params.AccountId);
+		}
+
 		if (!AccountInfoOSSAdapter)
 		{
 			InAsyncOp.SetError(Errors::InvalidUser());
