Fix Steam OSSAdapter for EOS ExternalAuth

diff --git a/Engine/Plugins/Online/OnlineServices/Source/OnlineServicesInterface/Private/Online/Auth.cpp b/Engine/Plugins/Online/OnlineServices/Source/OnlineServicesInterface/Private/Online/Auth.cpp
--- a/Engine/Plugins/Online/OnlineServices/Source/OnlineServicesInterface/Private/Online/Auth.cpp
+++ b/Engine/Plugins/Online/OnlineServices/Source/OnlineServicesInterface/Private/Online/Auth.cpp
@@ -19,7 +19,7 @@ const FName ExternalAuth = TEXT("ExternalAuth");
 namespace ExternalLoginType
 {
 const FName Epic = TEXT("Epic");
-const FName SteamAppTicket = TEXT("SteamAppTicket");
+const FName SteamSessionTicket = TEXT("SteamSessionTicket");
 const FName PsnIdToken = TEXT("PsnIdToken");
 const FName XblXstsToken = TEXT("XblXstsToken");
 const FName DiscordAccessToken = TEXT("DiscordAccessToken");
diff --git a/Engine/Plugins/Online/OnlineServices/Source/OnlineServicesInterface/Public/Online/Auth.h b/Engine/Plugins/Online/OnlineServices/Source/OnlineServicesInterface/Public/Online/Auth.h
--- a/Engine/Plugins/Online/OnlineServices/Source/OnlineServicesInterface/Public/Online/Auth.h
+++ b/Engine/Plugins/Online/OnlineServices/Source/OnlineServicesInterface/Public/Online/Auth.h
@@ -29,7 +29,7 @@ ONLINESERVICESINTERFACE_API extern const FName ExternalAuth;
 namespace ExternalLoginType
 {
 ONLINESERVICESINTERFACE_API extern const FName Epic;
-ONLINESERVICESINTERFACE_API extern const FName SteamAppTicket;
+ONLINESERVICESINTERFACE_API extern const FName SteamSessionTicket;
 ONLINESERVICESINTERFACE_API extern const FName PsnIdToken;
 ONLINESERVICESINTERFACE_API extern const FName XblXstsToken;
 ONLINESERVICESINTERFACE_API extern const FName DiscordAccessToken;
diff --git a/Engine/Plugins/Online/OnlineServicesEOSGS/Source/Private/Online/AuthEOSGS.cpp b/Engine/Plugins/Online/OnlineServicesEOSGS/Source/Private/Online/AuthEOSGS.cpp
--- a/Engine/Plugins/Online/OnlineServicesEOSGS/Source/Private/Online/AuthEOSGS.cpp
+++ b/Engine/Plugins/Online/OnlineServicesEOSGS/Source/Private/Online/AuthEOSGS.cpp
@@ -88,7 +88,7 @@ const FEOSConnectTranslationTraits* FEOSConnectLoginOptions::GetConnectTranslati
 {
 	static const TMap<FName, FEOSConnectTranslationTraits> SupportedLoginTranslatorTraits = {
 		{ ExternalLoginType::Epic, { EOS_EExternalCredentialType::EOS_ECT_EPIC, EEOSConnectTranslationFlags::None } },
-		{ ExternalLoginType::SteamAppTicket, { EOS_EExternalCredentialType::EOS_ECT_STEAM_APP_TICKET, EEOSConnectTranslationFlags::None } },
+		{ ExternalLoginType::SteamSessionTicket, { EOS_EExternalCredentialType::EOS_ECT_STEAM_SESSION_TICKET, EEOSConnectTranslationFlags::None } },
 		{ ExternalLoginType::PsnIdToken, { EOS_EExternalCredentialType::EOS_ECT_PSN_ID_TOKEN, EEOSConnectTranslationFlags::None } },
 		{ ExternalLoginType::XblXstsToken, { EOS_EExternalCredentialType::EOS_ECT_XBL_XSTS_TOKEN, EEOSConnectTranslationFlags::None } },
 		{ ExternalLoginType::DiscordAccessToken, { EOS_EExternalCredentialType::EOS_ECT_DISCORD_ACCESS_TOKEN, EEOSConnectTranslationFlags::None } },
@@ -257,7 +257,7 @@ const FEOSExternalAuthTranslationTraits* FEOSAuthLoginOptions::GetExternalAuthTr
 {
 	static const TMap<FName, FEOSExternalAuthTranslationTraits> SupportedExternalAuthTraits = {
 		{ ExternalLoginType::Epic, { EOS_EExternalCredentialType::EOS_ECT_EPIC } },
-		{ ExternalLoginType::SteamAppTicket, { EOS_EExternalCredentialType::EOS_ECT_STEAM_APP_TICKET } },
+		{ ExternalLoginType::SteamSessionTicket, { EOS_EExternalCredentialType::EOS_ECT_STEAM_SESSION_TICKET } },
 		{ ExternalLoginType::PsnIdToken, { EOS_EExternalCredentialType::EOS_ECT_PSN_ID_TOKEN } },
 		{ ExternalLoginType::XblXstsToken, { EOS_EExternalCredentialType::EOS_ECT_XBL_XSTS_TOKEN } },
 		{ ExternalLoginType::DiscordAccessToken, { EOS_EExternalCredentialType::EOS_ECT_DISCORD_ACCESS_TOKEN } },
diff --git a/Engine/Plugins/Online/OnlineServicesOSSAdapter/Source/Private/Online/AuthOSSAdapter.cpp b/Engine/Plugins/Online/OnlineServicesOSSAdapter/Source/Private/Online/AuthOSSAdapter.cpp
--- a/Engine/Plugins/Online/OnlineServicesOSSAdapter/Source/Private/Online/AuthOSSAdapter.cpp
+++ b/Engine/Plugins/Online/OnlineServicesOSSAdapter/Source/Private/Online/AuthOSSAdapter.cpp
@@ -64,7 +64,7 @@ const FExternalAuthTokenTranslationTraits* GetExternalAuthTokenTranslationTraits
 		{ TEXT("PS5"), { { ExternalLoginType::PsnIdToken, EExternalAuthTokenMethod::Primary, EExternalAuthTokenTranslationFlags::TokenString } } },
 		{ SWITCH_SUBSYSTEM, { { ExternalLoginType::NintendoNsaIdToken, EExternalAuthTokenMethod::Primary, EExternalAuthTokenTranslationFlags::AuthToken },
 							  { ExternalLoginType::NintendoIdToken, EExternalAuthTokenMethod::Secondary, EExternalAuthTokenTranslationFlags::TokenString } } },
-		{ STEAM_SUBSYSTEM, { { ExternalLoginType::SteamAppTicket, EExternalAuthTokenMethod::Primary, EExternalAuthTokenTranslationFlags::TokenBinary } } },
+		{ STEAM_SUBSYSTEM, { { ExternalLoginType::SteamSessionTicket, EExternalAuthTokenMethod::Primary, EExternalAuthTokenTranslationFlags::TokenString } } },
 	};
 
 	if (const TArray<FExternalAuthTokenTranslationTraits>* TraitsArray = SupportedExternalAuthTranslatorTraits.Find(SubsystemType))
@@ -267,6 +267,11 @@ TOnlineAsyncOpHandle<FAuthLogin> FAuthOSSAdapter::Login(FAuthLogin::Params&& Par
 			AccountInfoOSSAdapter->Attributes.Emplace(AccountAttributeData::DisplayName, UserOnlineAccount->GetDisplayName());
 			return MakeFulfilledPromise<void>().GetFuture();
 		}
+		else if (!GetIdentityInterface()->GetPlayerNickname(*AccountInfoOSSAdapter->UniqueNetId).IsEmpty())
+		{
+			AccountInfoOSSAdapter->Attributes.Emplace(AccountAttributeData::DisplayName, GetIdentityInterface()->GetPlayerNickname(*AccountInfoOSSAdapter->UniqueNetId));
+			return MakeFulfilledPromise<void>().GetFuture();
+		}
 		else
 		{
 			UE_LOG(LogOnlineServices, Warning, TEXT("[FAuthOSSAdapter::Login][%s] Failure: Failed to find UserOnlineAccount for account %s"), *GetSubsystem().GetSubsystemName().ToString(), *ToLogString(AccountInfoOSSAdapter->AccountId));
