Make AccountInfoRegistryOSSAdapter methods available outside of the OSSAdapter module.

#rnx

diff --git a/Engine/Plugins/Online/OnlineServicesOSSAdapter/Source/Public/Online/AuthOSSAdapter.h b/Engine/Plugins/Online/OnlineServicesOSSAdapter/Source/Public/Online/AuthOSSAdapter.h
--- a/Engine/Plugins/Online/OnlineServicesOSSAdapter/Source/Public/Online/AuthOSSAdapter.h
+++ b/Engine/Plugins/Online/OnlineServicesOSSAdapter/Source/Public/Online/AuthOSSAdapter.h
@@ -43,8 +43,8 @@ public:
 
 	virtual ~FAccountInfoRegistryOSSAdapter() = default;
 
-	TSharedPtr<FAccountInfoOSSAdapter> Find(FPlatformUserId PlatformUserId) const;
-	TSharedPtr<FAccountInfoOSSAdapter> Find(FAccountId AccountId) const;
+	ONLINESERVICESOSSADAPTER_API TSharedPtr<FAccountInfoOSSAdapter> Find(FPlatformUserId PlatformUserId) const;
+	ONLINESERVICESOSSADAPTER_API TSharedPtr<FAccountInfoOSSAdapter> Find(FAccountId AccountId) const;
 
 	void Register(const TSharedRef<FAccountInfoOSSAdapter>&UserAuthData);
 	void Unregister(FAccountId AccountId);
