Include InstanceName in the adapter OnlineSubsystem name
Enables adapter to work in PIE

diff --git a/Engine/Plugins/Online/OnlineServicesOSSAdapter/Source/Private/Online/OnlineServicesOSSAdapterModule.cpp b/Engine/Plugins/Online/OnlineServicesOSSAdapter/Source/Private/Online/OnlineServicesOSSAdapterModule.cpp
--- a/Engine/Plugins/Online/OnlineServicesOSSAdapter/Source/Private/Online/OnlineServicesOSSAdapterModule.cpp
+++ b/Engine/Plugins/Online/OnlineServicesOSSAdapter/Source/Private/Online/OnlineServicesOSSAdapterModule.cpp
@@ -61,7 +61,8 @@ public:
 	virtual ~FOnlineServicesFactoryOSSAdapter() {}
 	virtual TSharedPtr<IOnlineServices> Create(FName InInstanceName) override
 	{
-		IOnlineSubsystem* Subsystem = IOnlineSubsystem::Get(Config.OnlineSubsystem);
+		FName CombinedInstanceName = InInstanceName.IsNone() ? (/*NoSuffix*/Config.OnlineSubsystem) : FName(*FString::Printf(TEXT("%s:%s"), *Config.OnlineSubsystem.ToString(), *InInstanceName.ToString()));
+		IOnlineSubsystem* Subsystem = IOnlineSubsystem::Get(CombinedInstanceName);
 		if (Subsystem != nullptr)
 		{
 			return MakeShared<FOnlineServicesOSSAdapter>(Config.Service, Config.ConfigName, InInstanceName, Subsystem);
