[Backout] - CL34092249
Original CL Desc
-----------------------------------------------------------------
Include InstanceName in the adapter OnlineSubsystem name
Enables adapter to work in PIE

diff --git a/Engine/Plugins/Online/OnlineServicesOSSAdapter/Source/Private/Online/OnlineServicesOSSAdapterModule.cpp b/Engine/Plugins/Online/OnlineServicesOSSAdapter/Source/Private/Online/OnlineServicesOSSAdapterModule.cpp
--- a/Engine/Plugins/Online/OnlineServicesOSSAdapter/Source/Private/Online/OnlineServicesOSSAdapterModule.cpp
+++ b/Engine/Plugins/Online/OnlineServicesOSSAdapter/Source/Private/Online/OnlineServicesOSSAdapterModule.cpp
@@ -61,8 +61,7 @@ public:
 	virtual ~FOnlineServicesFactoryOSSAdapter() {}
 	virtual TSharedPtr<IOnlineServices> Create(FName InInstanceName) override
 	{
-		FName CombinedInstanceName = InInstanceName.IsNone() ? (/*NoSuffix*/Config.OnlineSubsystem) : FName(*FString::Printf(TEXT("%s:%s"), *Config.OnlineSubsystem.ToString(), *InInstanceName.ToString()));
-		IOnlineSubsystem* Subsystem = IOnlineSubsystem::Get(CombinedInstanceName);
+		IOnlineSubsystem* Subsystem = IOnlineSubsystem::Get(Config.OnlineSubsystem);
 		if (Subsystem != nullptr)
 		{
 			return MakeShared<FOnlineServicesOSSAdapter>(Config.Service, Config.ConfigName, InInstanceName, Subsystem);
