Implement OnSessionCreated and OnSessionJoined in the SessionsOSSAdapter so that an event is fired when a session is created and joined

diff --git a/Engine/Plugins/Online/OnlineServices/Source/OnlineServicesCommon/Private/Online/SessionsCommon.cpp b/Engine/Plugins/Online/OnlineServices/Source/OnlineServicesCommon/Private/Online/SessionsCommon.cpp
--- a/Engine/Plugins/Online/OnlineServices/Source/OnlineServicesCommon/Private/Online/SessionsCommon.cpp
+++ b/Engine/Plugins/Online/OnlineServices/Source/OnlineServicesCommon/Private/Online/SessionsCommon.cpp
@@ -886,6 +886,10 @@ namespace UE::Online {
 #undef CHECK_STATE_SESSION_NAME_ALREADY_EXISTS
 
 	// Events
+	TOnlineEvent<void(const FSessionCreated&)> FSessionsCommon::OnSessionCreated()
+	{
+		return SessionEvents.OnSessionCreated;
+	}
 
 	TOnlineEvent<void(const FSessionJoined&)> FSessionsCommon::OnSessionJoined()
 	{
diff --git a/Engine/Plugins/Online/OnlineServices/Source/OnlineServicesCommon/Public/Online/SessionsCommon.h b/Engine/Plugins/Online/OnlineServices/Source/OnlineServicesCommon/Public/Online/SessionsCommon.h
--- a/Engine/Plugins/Online/OnlineServices/Source/OnlineServicesCommon/Public/Online/SessionsCommon.h
+++ b/Engine/Plugins/Online/OnlineServices/Source/OnlineServicesCommon/Public/Online/SessionsCommon.h
@@ -270,6 +270,7 @@ public:
 	virtual TOnlineResult<FGetAllSessionInvites> GetAllSessionInvites(FGetAllSessionInvites::Params&& Params) override;
 	virtual TOnlineAsyncOpHandle<FRejectSessionInvite> RejectSessionInvite(FRejectSessionInvite::Params&& Params) override;
 
+	virtual TOnlineEvent<void(const FSessionCreated&)> OnSessionCreated() override;
 	virtual TOnlineEvent<void(const FSessionJoined&)> OnSessionJoined() override;
 	virtual TOnlineEvent<void(const FSessionLeft&)> OnSessionLeft() override;
 	virtual TOnlineEvent<void(const FSessionUpdated&)> OnSessionUpdated() override;
@@ -349,6 +350,7 @@ protected:
 
 	struct FSessionEvents
 	{
+		TOnlineEventCallable<void(const FSessionCreated&)> OnSessionCreated;
 		TOnlineEventCallable<void(const FSessionJoined&)> OnSessionJoined;
 		TOnlineEventCallable<void(const FSessionLeft&)> OnSessionLeft;
 		TOnlineEventCallable<void(const FSessionUpdated&)> OnSessionUpdated;
diff --git a/Engine/Plugins/Online/OnlineServices/Source/OnlineServicesInterface/Public/Online/Sessions.h b/Engine/Plugins/Online/OnlineServices/Source/OnlineServicesInterface/Public/Online/Sessions.h
--- a/Engine/Plugins/Online/OnlineServices/Source/OnlineServicesInterface/Public/Online/Sessions.h
+++ b/Engine/Plugins/Online/OnlineServices/Source/OnlineServicesInterface/Public/Online/Sessions.h
@@ -601,6 +601,15 @@ struct FRejectSessionInvite
 
 /* Events */
 
+struct FSessionCreated
+{
+	/** Id handle for the local user who created the session */
+	FAccountId LocalAccountId;
+
+	/** Id handle for the session created. */
+	FOnlineSessionId SessionId;
+};
+
 struct FSessionJoined
 {
 	/** Id handle for the local user who joined the session */
@@ -834,7 +843,13 @@ public:
 	virtual TOnlineAsyncOpHandle<FRejectSessionInvite> RejectSessionInvite(FRejectSessionInvite::Params&& Params) = 0;
 
 	/* Events */
-
+	
+	/**
+	 * This event will trigger as a result of a session being created.
+	 *
+	 * @return
+	 */
+	virtual TOnlineEvent<void(const FSessionCreated &)> OnSessionCreated() = 0;
 	/**
 	 * This event will trigger as a result of joining a session.
 	 *
@@ -1102,6 +1117,19 @@ BEGIN_ONLINE_STRUCT_META(FSessionUpdate)
 	ONLINE_STRUCT_FIELD(FSessionUpdate, RemovedSessionMembers)
 END_ONLINE_STRUCT_META()
 
+BEGIN_ONLINE_STRUCT_META(FSessionCreated)
+	ONLINE_STRUCT_FIELD(FSessionCreated, LocalAccountId),
+	ONLINE_STRUCT_FIELD(FSessionCreated, SessionId)
+END_ONLINE_STRUCT_META()
+
+BEGIN_ONLINE_STRUCT_META(FSessionJoined)
+ONLINE_STRUCT_FIELD(FSessionCreated, LocalAccountId),
+ONLINE_STRUCT_FIELD(FSessionCreated, SessionId)
+END_ONLINE_STRUCT_META()
+
+BEGIN_ONLINE_STRUCT_META(FSessionLeft)
+	ONLINE_STRUCT_FIELD(FSessionCreated, LocalAccountId)
+END_ONLINE_STRUCT_META()
 /* Meta*/ }
 
 /* UE::Online */ }
diff --git a/Engine/Plugins/Online/OnlineServicesOSSAdapter/Source/Private/Online/SessionsOSSAdapter.cpp b/Engine/Plugins/Online/OnlineServicesOSSAdapter/Source/Private/Online/SessionsOSSAdapter.cpp
--- a/Engine/Plugins/Online/OnlineServicesOSSAdapter/Source/Private/Online/SessionsOSSAdapter.cpp
+++ b/Engine/Plugins/Online/OnlineServicesOSSAdapter/Source/Private/Online/SessionsOSSAdapter.cpp
@@ -319,6 +319,34 @@ void FSessionsOSSAdapter::Initialize()
 		SessionEvents.OnUISessionJoinRequested.Broadcast(Event);
 	}));
 
+	SessionsInterface->AddOnCreateSessionCompleteDelegate_Handle(FOnCreateSessionCompleteDelegate::CreateLambda([this](FName SessionName, bool bSessionCompletionResult)
+	{
+		if (bSessionCompletionResult)
+		{
+			FAccountId AccountId;
+			FOnlineSessionId OnlineSessionId;
+			if (GetAccountIdAndSessionIdFromNamedSession(SessionName, AccountId, OnlineSessionId))
+			{
+				const FSessionCreated Event { AccountId, OnlineSessionId };
+				SessionEvents.OnSessionCreated.Broadcast(Event);
+			}
+		}
+	}));
+
+	SessionsInterface->AddOnJoinSessionCompleteDelegate_Handle(FOnJoinSessionCompleteDelegate::CreateLambda([this](FName SessionName, EOnJoinSessionCompleteResult::Type JoinSessionType)
+	{
+		if (JoinSessionType == EOnJoinSessionCompleteResult::Type::Success)
+		{
+			FAccountId AccountId;
+			FOnlineSessionId OnlineSessionId;
+			if (GetAccountIdAndSessionIdFromNamedSession(SessionName, AccountId, OnlineSessionId))
+			{
+				const FSessionJoined Event { AccountId, OnlineSessionId };
+				SessionEvents.OnSessionJoined.Broadcast(Event);
+			}
+		}
+	}));
+
 	SessionsInterface->AddOnSessionParticipantJoinedDelegate_Handle(FOnSessionParticipantJoinedDelegate::CreateLambda([this](FName SessionName, const FUniqueNetId& TargetUniqueNetId)
 	{
 		// We won't update a session that can't be retrieved from the OSS or doesn't exist in OnlineServices anymore
@@ -365,6 +393,8 @@ void FSessionsOSSAdapter::Shutdown()
 	SessionsInterface->ClearOnSessionParticipantJoinedDelegates(this);
 	SessionsInterface->ClearOnSessionParticipantLeftDelegates(this);
 	SessionsInterface->ClearOnSessionParticipantSettingsUpdatedDelegates(this);
+	SessionsInterface->ClearOnCreateSessionCompleteDelegates(this);
+	SessionsInterface->ClearOnJoinSessionCompleteDelegates(this);
 
 	Super::Shutdown();
 }
@@ -1149,6 +1179,52 @@ TArray<TSharedRef<FSessionCommon>> FSessionsOSSAdapter::BuildV2SessionSearchResu
 	return FoundSessions;
 }
 
+bool FSessionsOSSAdapter::GetAccountIdAndSessionIdFromNamedSession(const FName& SessionName, FAccountId& AccountId, FOnlineSessionId& OnlineSessionId) const
+{
+	if (GetV2SessionCreatorFromV1NamedSession(SessionName, AccountId))
+	{
+		if (GetV2SessionIdFromV1NamedSession(SessionName, OnlineSessionId))
+		{
+			return true;
+		}
+	}
+
+	return false;
+}
+
+bool FSessionsOSSAdapter::GetV2SessionIdFromV1NamedSession(const FName& SessionName, FOnlineSessionId& OnlineSessionId) const
+{
+	const FNamedOnlineSession* NamedOnlineSession = SessionsInterface->GetNamedSession(SessionName);
+	if (NamedOnlineSession)
+	{
+		const FUniqueNetIdRef SessionUniqueNetIdRef = NamedOnlineSession->SessionInfo->GetSessionId().AsShared();
+		OnlineSessionId = GetSessionIdRegistry().FindOrAddHandle(SessionUniqueNetIdRef);
+		return true;
+	}
+
+	return false;
+}
+
+bool FSessionsOSSAdapter::GetV2SessionCreatorFromV1NamedSession(const FName& SessionName, FAccountId& AccountId) const
+{
+	const FNamedOnlineSession* NamedOnlineSession = SessionsInterface->GetNamedSession(SessionName);
+	if (NamedOnlineSession)
+	{
+		const FUniqueNetIdPtr LocalOwnerIdPtr = NamedOnlineSession->OwningUserId;
+		if (LocalOwnerIdPtr.IsValid())
+		{
+			const FOnlineServicesOSSAdapter& ServicesOSSAdapter = static_cast<FOnlineServicesOSSAdapter&>(Services);
+			AccountId = ServicesOSSAdapter.GetAccountIdRegistry().FindOrAddHandle(LocalOwnerIdPtr->AsShared());
+			if (AccountId.IsValid())
+			{
+				return true;
+			}
+		}
+	}
+
+	return false;
+}
+
 FOnlineSessionIdRegistryOSSAdapter& FSessionsOSSAdapter::GetSessionIdRegistry() const
 {
 	FOnlineSessionIdRegistryOSSAdapter* SessionIdRegistry = static_cast<FOnlineSessionIdRegistryOSSAdapter*>(FOnlineIdRegistryRegistry::Get().GetSessionIdRegistry(Services.GetServicesProvider()));
diff --git a/Engine/Plugins/Online/OnlineServicesOSSAdapter/Source/Public/Online/SessionsOSSAdapter.h b/Engine/Plugins/Online/OnlineServicesOSSAdapter/Source/Public/Online/SessionsOSSAdapter.h
--- a/Engine/Plugins/Online/OnlineServicesOSSAdapter/Source/Public/Online/SessionsOSSAdapter.h
+++ b/Engine/Plugins/Online/OnlineServicesOSSAdapter/Source/Public/Online/SessionsOSSAdapter.h
@@ -101,6 +101,13 @@ private:
 	/** Builds a V2 Session Search Results array from the passed V1 equivalent types */
 	TArray<TSharedRef<FSessionCommon>> BuildV2SessionSearchResults(const TArray<FOnlineSessionSearchResult>& SessionSearchResults) const;
 
+	/** Helper to get access to the account id and session id for V1 named sessions */
+	bool GetAccountIdAndSessionIdFromNamedSession(const FName& SessionName, FAccountId& AccountId, FOnlineSessionId& OnlineSessionId) const;
+	/** Helper to get the V2 session id from a v1 named session */
+	bool GetV2SessionIdFromV1NamedSession(const FName& SessionName, FOnlineSessionId& OnlineSessionId) const;
+	/** Helper to get the V2 session creator account id from a v1 named session */
+	bool GetV2SessionCreatorFromV1NamedSession(const FName& SessionName, FAccountId& AccountId) const;
+
 	FOnlineSessionIdRegistryOSSAdapter& GetSessionIdRegistry() const;
 	FOnlineSessionInviteIdRegistryOSSAdapter& GetSessionInviteIdRegistry() const;
 
