OSS adapter compatibility changes part 1

diff --git a/Engine/Plugins/Online/OnlineServices/Source/OnlineServicesCommon/Private/Online/SessionsCommon.cpp b/Engine/Plugins/Online/OnlineServices/Source/OnlineServicesCommon/Private/Online/SessionsCommon.cpp
--- a/Engine/Plugins/Online/OnlineServices/Source/OnlineServicesCommon/Private/Online/SessionsCommon.cpp
+++ b/Engine/Plugins/Online/OnlineServices/Source/OnlineServicesCommon/Private/Online/SessionsCommon.cpp
@@ -281,7 +281,6 @@ namespace UE::Online {
 	TOnlineResult<FGetMutableSessionById> FSessionsCommon::GetMutableSessionById(FGetMutableSessionById::Params&& Params) const
 	{
 		CHECK_PARAMS_ID_HANDLE(Params.SessionId, FGetMutableSessionById)
-
 		if (const TSharedRef<FSessionCommon>* FoundSession = AllSessionsById.Find(Params.SessionId))
 		{
 			return TOnlineResult<FGetMutableSessionById>({ *FoundSession });
@@ -670,7 +669,7 @@ namespace UE::Online {
 		TOnlineResult<FGetSessionById> GetSessionByIdResult = GetSessionById({ Params.SessionId });
 		if (GetSessionByIdResult.IsError())
 		{
-			UE_LOG(LogOnlineServices, Verbose, TEXT("[%s] Session [%s] not found. Please call FindSessions to get an updated list of available sessions "), UTF8_TO_TCHAR(__FUNCTION__), *ToLogString(Params.SessionId));
+			UE_LOG(LogOnlineServices, Verbose, TEXT("[%s] Session [%s] not found. Please call FindSessions to get an updated list of available sessions. Invites should be cached before joining "), UTF8_TO_TCHAR(__FUNCTION__), *ToLogString(Params.SessionId));
 			return TOptional<FOnlineError>(GetSessionByIdResult.GetErrorValue());
 		}
 
@@ -689,7 +688,7 @@ namespace UE::Online {
 			UE_LOG(LogOnlineServices, Verbose, TEXT("[%s] Session [%s] not joinable "), UTF8_TO_TCHAR(__FUNCTION__), *ToLogString(Params.SessionId));
 			return TOptional<FOnlineError>(Errors::AccessDenied());
 		}
-
+		
 		if (Params.bPresenceEnabled)
 		{
 			for (const TPair<FName, FOnlineSessionId>& Entry : LocalSessionsByName)
diff --git a/Engine/Plugins/Online/OnlineServicesOSSAdapter/Source/Private/Online/SessionsOSSAdapter.cpp b/Engine/Plugins/Online/OnlineServicesOSSAdapter/Source/Private/Online/SessionsOSSAdapter.cpp
--- a/Engine/Plugins/Online/OnlineServicesOSSAdapter/Source/Private/Online/SessionsOSSAdapter.cpp
+++ b/Engine/Plugins/Online/OnlineServicesOSSAdapter/Source/Private/Online/SessionsOSSAdapter.cpp
@@ -243,8 +243,9 @@ TSharedRef<FOnlineSessionSearch> BuildV1SessionSearch(const FFindSessions::Param
 void WriteV2SessionInfoFromV1Session(const FOnlineSession& InSession, FSessionInfo& OutInfo)
 {
 	OutInfo.SessionIdOverride = InSession.SessionSettings.SessionIdOverride;
-
-	InSession.SessionSettings.Settings.FindChecked(OSS_ADAPTER_SESSIONS_ALLOW_SANCTIONED_PLAYERS).Data.GetValue(OutInfo.bAllowSanctionedPlayers);
+	bool bAllowSanctionedPlayers = false;
+	InSession.SessionSettings.Get(OSS_ADAPTER_SESSIONS_ALLOW_SANCTIONED_PLAYERS, bAllowSanctionedPlayers);
+	OutInfo.bAllowSanctionedPlayers = bAllowSanctionedPlayers;
 
 	OutInfo.bAntiCheatProtected = InSession.SessionSettings.bAntiCheatProtected;
 
@@ -310,8 +311,12 @@ void FSessionsOSSAdapter::Initialize()
 			LocalAccountId = (*LocalAccountIdPtr).AsShared();
 		}
 
+		TSharedRef<FSessionCommon> V2Session = BuildV2Session(&InviteResult.Session);
+		const FAccountId LocalAccountIdV2 = ServicesOSSAdapter.GetAccountIdRegistry().FindOrAddHandle(LocalAccountId);
+		AddSearchResult(V2Session, LocalAccountIdV2);
+
 		FUISessionJoinRequested Event {
-			ServicesOSSAdapter.GetAccountIdRegistry().FindOrAddHandle(LocalAccountId),
+			LocalAccountIdV2,
 			TResult<FOnlineSessionId, FOnlineError>(GetSessionIdRegistry().FindOrAddHandle(InviteResult.Session.SessionInfo->GetSessionId().AsShared())),
 			EUISessionJoinRequestedSource::FromInvitation
 		};
@@ -1103,7 +1108,7 @@ void FSessionsOSSAdapter::WriteV2SessionSettingsFromV1Session(const FOnlineSessi
 	OutSession->SessionSettings.NumMaxConnections = FMath::Max(InSession->SessionSettings.NumPrivateConnections, InSession->SessionSettings.NumPublicConnections);
 
 	FString SchemaNameStr;
-	InSession->SessionSettings.Settings.FindChecked(OSS_ADAPTER_SESSIONS_SCHEMA_NAME).Data.GetValue(SchemaNameStr);
+	InSession->SessionSettings.Get(OSS_ADAPTER_SESSIONS_SCHEMA_NAME, SchemaNameStr);
 	OutSession->SessionSettings.SchemaName = FSchemaId(SchemaNameStr);
 
 	OutSession->SessionSettings.CustomSettings = GetV2SessionSettings(InSession->SessionSettings.Settings);
@@ -1144,17 +1149,23 @@ FOnlineSession FSessionsOSSAdapter::BuildV1Session(const TSharedRef<const ISessi
 	return FOnlineSession(FSessionOSSAdapter::Cast(*InSession).GetV1Session());
 }
 
-TSharedRef<FSessionCommon> FSessionsOSSAdapter::BuildV2Session(const FOnlineSession* InSession) const
+TSharedRef<FSessionCommon> FSessionsOSSAdapter::BuildV2Session(const FNamedOnlineSession* InSession) const
 {
 	FOnlineServicesOSSAdapter& ServicesOSSAdapter = static_cast<FOnlineServicesOSSAdapter&>(Services);
 
 	TSharedRef<FSessionOSSAdapter> Session = MakeShared<FSessionOSSAdapter>(*InSession, GetSessionIdRegistry().FindOrAddHandle(InSession->SessionInfo->GetSessionId().AsShared()));
 
-	if (const FNamedOnlineSession* NamedInSession = static_cast<const FNamedOnlineSession*>(InSession))
-	{
-		Session->OwnerAccountId = ServicesOSSAdapter.GetAccountIdRegistry().FindOrAddHandle(NamedInSession->LocalOwnerId->AsShared());
-		WriteV2SessionSettingsFromV1NamedSession(NamedInSession, Session);
-	}
+	Session->OwnerAccountId = ServicesOSSAdapter.GetAccountIdRegistry().FindOrAddHandle(InSession->LocalOwnerId->AsShared());
+	WriteV2SessionSettingsFromV1NamedSession(InSession, Session);
+
+	return Session;
+}
+
+TSharedRef<FSessionCommon> FSessionsOSSAdapter::BuildV2Session(const FOnlineSession* InSession) const
+{
+	FOnlineServicesOSSAdapter& ServicesOSSAdapter = static_cast<FOnlineServicesOSSAdapter&>(Services);
+
+	TSharedRef<FSessionOSSAdapter> Session = MakeShared<FSessionOSSAdapter>(*InSession, GetSessionIdRegistry().FindOrAddHandle(InSession->SessionInfo->GetSessionId().AsShared()));
 
 	WriteV2SessionSettingsFromV1Session(InSession, Session);
 
diff --git a/Engine/Plugins/Online/OnlineServicesOSSAdapter/Source/Public/Online/SessionsOSSAdapter.h b/Engine/Plugins/Online/OnlineServicesOSSAdapter/Source/Public/Online/SessionsOSSAdapter.h
--- a/Engine/Plugins/Online/OnlineServicesOSSAdapter/Source/Public/Online/SessionsOSSAdapter.h
+++ b/Engine/Plugins/Online/OnlineServicesOSSAdapter/Source/Public/Online/SessionsOSSAdapter.h
@@ -97,6 +97,7 @@ private:
 	/** Builds a V1 Session data type from a V2 equivalent */
 	FOnlineSession BuildV1Session(const TSharedRef<const ISession> InSession) const;
 	/** Builds a V2 Session data type from a V1 equivalent */
+	TSharedRef<FSessionCommon> BuildV2Session(const FNamedOnlineSession* InSession) const;
 	TSharedRef<FSessionCommon> BuildV2Session(const FOnlineSession* InSession) const;
 	/** Builds a V2 Session Search Results array from the passed V1 equivalent types */
 	TArray<TSharedRef<FSessionCommon>> BuildV2SessionSearchResults(const TArray<FOnlineSessionSearchResult>& SessionSearchResults) const;
diff --git a/Engine/Source/Runtime/Engine/Private/OnlineReplStructs.cpp b/Engine/Source/Runtime/Engine/Private/OnlineReplStructs.cpp
--- a/Engine/Source/Runtime/Engine/Private/OnlineReplStructs.cpp
+++ b/Engine/Source/Runtime/Engine/Private/OnlineReplStructs.cpp
@@ -235,7 +235,7 @@ void FUniqueNetIdRepl::MakeReplicationDataV2()
 	UE::Online::FAccountId AccountId = GetV2();
 	UE::Online::EOnlineServices OnlineServicesType = AccountId.GetOnlineServicesType();
 	TArray<uint8> ReplicationData = UE::Online::FOnlineIdRegistryRegistry::Get().ToReplicationData(AccountId);
-	check(!ReplicationData.IsEmpty());
+	checkSlow(!ReplicationData.IsEmpty());
 
 	const int32 TotalBytes = sizeof(EncodingFlags) + sizeof(OnlineServicesType) + sizeof(TArray<uint8>::SizeType) + ReplicationData.Num();
 	ReplicationBytes.Empty(TotalBytes);
@@ -452,7 +452,7 @@ void FUniqueNetIdRepl::NetSerializeLoadV2(FArchive& Ar, const EUniqueIdEncodingF
 	Ar << ReplicationData;
 
 	const UE::Online::FAccountId AccountId = UE::Online::FOnlineIdRegistryRegistry::Get().ToAccountId(OnlineServicesType, ReplicationData);
-	check(AccountId.IsValid());
+	checkSlow(AccountId.IsValid());
 	SetAccountId(AccountId);
 }
 
