Online services can be hit from Verse so needs to be instrumented with AutoRTFM.

#rnx

diff --git a/Engine/Plugins/Online/OnlineServices/Source/OnlineServicesCommon/OnlineServicesCommon.Build.cs b/Engine/Plugins/Online/OnlineServices/Source/OnlineServicesCommon/OnlineServicesCommon.Build.cs
--- a/Engine/Plugins/Online/OnlineServices/Source/OnlineServicesCommon/OnlineServicesCommon.Build.cs
+++ b/Engine/Plugins/Online/OnlineServices/Source/OnlineServicesCommon/OnlineServicesCommon.Build.cs
@@ -24,7 +24,5 @@ public class OnlineServicesCommon : ModuleRules
 				"CoreUObject"		// CoreUObject temporary dependency
 			}
 		);
-
-		bDisableAutoRTFMInstrumentation = true;
 	}
 }
diff --git a/Engine/Plugins/Online/OnlineServicesEOS/Source/OnlineServicesEOS.Build.cs b/Engine/Plugins/Online/OnlineServicesEOS/Source/OnlineServicesEOS.Build.cs
--- a/Engine/Plugins/Online/OnlineServicesEOS/Source/OnlineServicesEOS.Build.cs
+++ b/Engine/Plugins/Online/OnlineServicesEOS/Source/OnlineServicesEOS.Build.cs
@@ -24,7 +24,5 @@ public class OnlineServicesEOS : ModuleRules
 				"EOSShared",
 			}
 		);
-
-		bDisableAutoRTFMInstrumentation = true;
 	}
 }
diff --git a/Engine/Plugins/Online/OnlineServicesEOSGS/Source/OnlineServicesEOSGS.Build.cs b/Engine/Plugins/Online/OnlineServicesEOSGS/Source/OnlineServicesEOSGS.Build.cs
--- a/Engine/Plugins/Online/OnlineServicesEOSGS/Source/OnlineServicesEOSGS.Build.cs
+++ b/Engine/Plugins/Online/OnlineServicesEOSGS/Source/OnlineServicesEOSGS.Build.cs
@@ -42,7 +42,5 @@ public class OnlineServicesEOSGS : ModuleRules
 			PublicWeakFrameworks.Add("AuthenticationServices");
 			PrivateDependencyModuleNames.Add("ApplicationCore");
 		}
-
-		bDisableAutoRTFMInstrumentation = true;
 	}
 }
diff --git a/Engine/Plugins/Online/OnlineServicesOSSAdapter/Source/OnlineServicesOSSAdapter.Build.cs b/Engine/Plugins/Online/OnlineServicesOSSAdapter/Source/OnlineServicesOSSAdapter.Build.cs
--- a/Engine/Plugins/Online/OnlineServicesOSSAdapter/Source/OnlineServicesOSSAdapter.Build.cs
+++ b/Engine/Plugins/Online/OnlineServicesOSSAdapter/Source/OnlineServicesOSSAdapter.Build.cs
@@ -22,7 +22,5 @@ public class OnlineServicesOSSAdapter : ModuleRules
 				"CoreOnline",
 			}
 		);
-
-		bDisableAutoRTFMInstrumentation = true;
 	}
 }
diff --git a/Engine/Plugins/Online/OnlineServicesOSSAdapter/Source/Public/Online/OnlineIdOSSAdapter.h b/Engine/Plugins/Online/OnlineServicesOSSAdapter/Source/Public/Online/OnlineIdOSSAdapter.h
--- a/Engine/Plugins/Online/OnlineServicesOSSAdapter/Source/Public/Online/OnlineIdOSSAdapter.h
+++ b/Engine/Plugins/Online/OnlineServicesOSSAdapter/Source/Public/Online/OnlineIdOSSAdapter.h
@@ -4,7 +4,7 @@
 
 #include "Online/CoreOnline.h"
 #include "OnlineSubsystemTypes.h"
-#include "Misc/ScopeRWLock.h"
+#include "Misc/TransactionallySafeRWLock.h"
 
 namespace UE::Online {
 
@@ -32,7 +32,7 @@ public:
 
 		// Take a read lock and check if we already have a handle
 		{
-			FReadScopeLock ReadLock(Lock);
+			UE::TReadScopeLock ReadLock(Lock);
 			if (const uint32* FoundHandle = IdValueToHandleMap.Find(IdValue))
 			{
 				OnlineId = TOnlineId<IdType>(OnlineServicesType, *FoundHandle);
@@ -42,7 +42,7 @@ public:
 		if (!OnlineId.IsValid())
 		{
 			// Take a write lock, check again if we already have a handle, or insert a new element.
-			FWriteScopeLock WriteLock(Lock);
+			UE::TWriteScopeLock WriteLock(Lock);
 			if (const uint32* FoundHandle = IdValueToHandleMap.Find(IdValue))
 			{
 				OnlineId = TOnlineId<IdType>(OnlineServicesType, *FoundHandle);
@@ -64,7 +64,7 @@ public:
 	{
 		if (OnlineId.GetOnlineServicesType() == OnlineServicesType && OnlineId.IsValid())
 		{
-			FReadScopeLock ReadLock(Lock);
+			UE::TReadScopeLock ReadLock(Lock);
 			if (IdValues.IsValidIndex(OnlineId.GetHandle() - 1))
 			{
 				return IdValues[OnlineId.GetHandle() - 1];
@@ -109,7 +109,7 @@ public:
 	// End IOnlineAccountIdRegistry
 
 private:
-	mutable FRWLock Lock;
+	mutable FTransactionallySafeRWLock Lock;
 
 	EOnlineServices OnlineServicesType;
 	TArray<FUniqueNetIdRef> IdValues;
