OnlineTests on PS4/PS5: Fixed accessing to non valid pointer at the test suite completion.

diff --git a/Engine/Plugins/Online/OnlineServicesOSSAdapter/Source/Private/Online/OnlineServicesOSSAdapterModule.cpp b/Engine/Plugins/Online/OnlineServicesOSSAdapter/Source/Private/Online/OnlineServicesOSSAdapterModule.cpp
--- a/Engine/Plugins/Online/OnlineServicesOSSAdapter/Source/Private/Online/OnlineServicesOSSAdapterModule.cpp
+++ b/Engine/Plugins/Online/OnlineServicesOSSAdapter/Source/Private/Online/OnlineServicesOSSAdapterModule.cpp
@@ -103,13 +103,10 @@ void FOnlineServicesOSSAdapterModule::ShutdownModule()
 	{
 		for (const FOSSAdapterService& ServiceConfig : Config.Services)
 		{
-			if (IOnlineSubsystem::IsEnabled(ServiceConfig.OnlineSubsystem))
-			{
-				FOnlineServicesRegistry::Get().UnregisterServicesFactory(ServiceConfig.Service, ServiceConfig.Priority);
-				FOnlineIdRegistryRegistry::Get().UnregisterAccountIdRegistry(ServiceConfig.Service, ServiceConfig.Priority);
-				FOnlineIdRegistryRegistry::Get().UnregisterSessionIdRegistry(ServiceConfig.Service, ServiceConfig.Priority);
-				FOnlineIdRegistryRegistry::Get().UnregisterSessionInviteIdRegistry(ServiceConfig.Service, ServiceConfig.Priority);
-			}
+			FOnlineServicesRegistry::Get().UnregisterServicesFactory(ServiceConfig.Service, ServiceConfig.Priority);
+			FOnlineIdRegistryRegistry::Get().UnregisterAccountIdRegistry(ServiceConfig.Service, ServiceConfig.Priority);
+			FOnlineIdRegistryRegistry::Get().UnregisterSessionIdRegistry(ServiceConfig.Service, ServiceConfig.Priority);
+			FOnlineIdRegistryRegistry::Get().UnregisterSessionInviteIdRegistry(ServiceConfig.Service, ServiceConfig.Priority);
 		}
 	}
 }
