Fix out parameter in presence adapter PropertiesV1toV2

diff --git a/Engine/Plugins/Online/OnlineServicesOSSAdapter/Source/Private/Online/PresenceOSSAdapter.cpp b/Engine/Plugins/Online/OnlineServicesOSSAdapter/Source/Private/Online/PresenceOSSAdapter.cpp
--- a/Engine/Plugins/Online/OnlineServicesOSSAdapter/Source/Private/Online/PresenceOSSAdapter.cpp
+++ b/Engine/Plugins/Online/OnlineServicesOSSAdapter/Source/Private/Online/PresenceOSSAdapter.cpp
@@ -81,7 +81,7 @@ EUserPresenceGameStatus GameStatusV1toV2(FOnlineUserPresence& Presence)
 	}
 }
 
-void PropertiesV1toV2(FOnlineKeyValuePairs<FPresenceKey, FVariantData>& PropertiesIn, TMap<FString, FPresenceVariant> PropertiesOut)
+void PropertiesV1toV2(FOnlineKeyValuePairs<FPresenceKey, FVariantData>& PropertiesIn, TMap<FString, FPresenceVariant>& PropertiesOut)
 {
 	for (const TPair<FPresenceKey, FVariantData>& Pair : PropertiesIn)
 	{
