Add the logged in user to registry in auth oss adaptor instead of return InvalidUser error when the account logged in directly through ossv1.

diff --git a/Engine/Plugins/Online/OnlineServicesOSSAdapter/Source/Private/Online/AuthOSSAdapter.cpp b/Engine/Plugins/Online/OnlineServicesOSSAdapter/Source/Private/Online/AuthOSSAdapter.cpp
--- a/Engine/Plugins/Online/OnlineServicesOSSAdapter/Source/Private/Online/AuthOSSAdapter.cpp
+++ b/Engine/Plugins/Online/OnlineServicesOSSAdapter/Source/Private/Online/AuthOSSAdapter.cpp
@@ -21,6 +21,12 @@
 
 namespace UE::Online {
 
+static TAutoConsoleVariable<bool> CVarAuthRegisterAccountLoggedInThroughOSSEnabled(
+	TEXT("Online.AuthRegisterAccountLoggedInThroughOSSEnabled"),
+	true,
+	TEXT("When enabled, will register account logged in directly through OSSv1 instead of returning error"),
+	ECVF_Default);
+
 namespace {
 
 ELoginStatus TranslateLoginStatus(::ELoginStatus::Type Status)
@@ -263,19 +269,8 @@ TOnlineAsyncOpHandle<FAuthLogin> FAuthOSSAdapter::Login(FAuthLogin::Params&& Par
 	{
 		const TSharedRef<FAccountInfoOSSAdapter>& AccountInfoOSSAdapter = GetOpDataChecked<TSharedRef<FAccountInfoOSSAdapter>>(InAsyncOp, UE_ONLINE_AUTH_OSS_ADAPTER_KEY_NAME_ACCOUNT_INFO);
 
-		if (TSharedPtr<FUserOnlineAccount> UserOnlineAccount = GetIdentityInterface()->GetUserAccount(*AccountInfoOSSAdapter->UniqueNetId))
-		{
-			AccountInfoOSSAdapter->Attributes.Emplace(AccountAttributeData::DisplayName, UserOnlineAccount->GetDisplayName());
-			return MakeFulfilledPromise<void>().GetFuture();
-		}
-		else if (!GetIdentityInterface()->GetPlayerNickname(*AccountInfoOSSAdapter->UniqueNetId).IsEmpty())
+		if (!PopulateAttributes(AccountInfoOSSAdapter))
 		{
-			AccountInfoOSSAdapter->Attributes.Emplace(AccountAttributeData::DisplayName, GetIdentityInterface()->GetPlayerNickname(*AccountInfoOSSAdapter->UniqueNetId));
-			return MakeFulfilledPromise<void>().GetFuture();
-		}
-		else
-		{
-			UE_LOG(LogOnlineServices, Warning, TEXT("[FAuthOSSAdapter::Login][%s] Failure: Failed to find UserOnlineAccount for account %s"), *GetSubsystem().GetSubsystemName().ToString(), *ToLogString(AccountInfoOSSAdapter->AccountId));
 			InAsyncOp.SetError(Errors::Unknown());
 
 			// Logout.
@@ -294,6 +289,8 @@ TOnlineAsyncOpHandle<FAuthLogin> FAuthOSSAdapter::Login(FAuthLogin::Params&& Par
 			GetIdentityInterface()->Logout(AccountInfoOSSAdapter->LocalUserNum);
 			return Future;
 		}
+
+		return MakeFulfilledPromise<void>().GetFuture();
 	})
 	// Step 6: bookkeeping and notifications.
 	.Then([this](TOnlineAsyncOp<FAuthLogin>& InAsyncOp)
@@ -601,8 +598,27 @@ TOnlineAsyncOpHandle<FAuthHandleLoginStatusChangedImpl> FAuthOSSAdapter::HandleL
 
 		if (!AccountInfoOSSAdapter)
 		{
-			InAsyncOp.SetError(Errors::InvalidUser());
-			return;
+			if (CVarAuthRegisterAccountLoggedInThroughOSSEnabled.GetValueOnAnyThread())
+			{
+				// In case of Login get called without going through the OSSv2 OSS Adapter, add logged in user to AccountInfoRegistry as well
+				AccountInfoOSSAdapter = MakeShared<FAccountInfoOSSAdapter>();
+				AccountInfoOSSAdapter->AccountId = Params.AccountId;
+				AccountInfoOSSAdapter->PlatformUserId = Params.PlatformUserId;
+				AccountInfoOSSAdapter->LoginStatus = ELoginStatus::NotLoggedIn;
+				AccountInfoOSSAdapter->LocalUserNum = FPlatformMisc::GetUserIndexForPlatformUser(Params.PlatformUserId);
+				AccountInfoOSSAdapter->UniqueNetId = GetIdentityInterface()->GetUniquePlayerId(AccountInfoOSSAdapter->LocalUserNum);
+				if (!PopulateAttributes(AccountInfoOSSAdapter))
+				{
+					InAsyncOp.SetError(Errors::InvalidUser());
+					return;
+				}
+				AccountInfoRegistryOSSAdapter.Register(AccountInfoOSSAdapter.ToSharedRef());
+			}
+			else
+			{
+				InAsyncOp.SetError(Errors::InvalidUser());
+				return;
+			}
 		}
 
 		// Don't send duplicate notification during login
@@ -634,4 +650,26 @@ TOnlineAsyncOpHandle<FAuthHandleLoginStatusChangedImpl> FAuthOSSAdapter::HandleL
 	return Op->GetHandle();
 }
 
+bool FAuthOSSAdapter::PopulateAttributes(const TSharedPtr<FAccountInfoOSSAdapter>& AccountInfoOSSAdapter) const
+{
+	TSharedPtr<FUserOnlineAccount> UserOnlineAccount = GetIdentityInterface()->GetUserAccount(*AccountInfoOSSAdapter->UniqueNetId);
+	if (UserOnlineAccount)
+	{
+		AccountInfoOSSAdapter->Attributes.Emplace(AccountAttributeData::DisplayName, UserOnlineAccount->GetDisplayName());
+		return true;
+	}
+
+	FString Displayname = GetIdentityInterface()->GetPlayerNickname(*AccountInfoOSSAdapter->UniqueNetId);
+
+	if (!Displayname.IsEmpty())
+	{
+		AccountInfoOSSAdapter->Attributes.Emplace(AccountAttributeData::DisplayName, GetIdentityInterface()->GetPlayerNickname(*AccountInfoOSSAdapter->UniqueNetId));
+		return true;
+	}
+
+	UE_LOG(LogOnlineServices, Warning, TEXT("[FAuthOSSAdapter::Login][%s] Failure: Failed to find UserOnlineAccount for account %s"), *GetSubsystem().GetSubsystemName().ToString(), *ToLogString(AccountInfoOSSAdapter->AccountId));
+
+	return false;
+}
+
 /* UE::Online */ }
diff --git a/Engine/Plugins/Online/OnlineServicesOSSAdapter/Source/Public/Online/AuthOSSAdapter.h b/Engine/Plugins/Online/OnlineServicesOSSAdapter/Source/Public/Online/AuthOSSAdapter.h
--- a/Engine/Plugins/Online/OnlineServicesOSSAdapter/Source/Public/Online/AuthOSSAdapter.h
+++ b/Engine/Plugins/Online/OnlineServicesOSSAdapter/Source/Public/Online/AuthOSSAdapter.h
@@ -87,6 +87,8 @@ protected:
 
 	TOnlineAsyncOpHandle<FAuthHandleLoginStatusChangedImpl> HandleLoginStatusChangedImplOp(FAuthHandleLoginStatusChangedImpl::Params&& Params);
 
+	bool PopulateAttributes(const TSharedPtr<FAccountInfoOSSAdapter>& AccountInfoOSSAdapter) const;
+
 	FAccountInfoRegistryOSSAdapter AccountInfoRegistryOSSAdapter;
 	FDelegateHandle OnLoginStatusChangedHandle[MAX_LOCAL_PLAYERS];
 };
