Add safe check in AuthOSSAdapter when populating attributes for new account, in case UniqueNetId is null.

diff --git a/Engine/Plugins/Online/OnlineServicesOSSAdapter/Source/Private/Online/AuthOSSAdapter.cpp b/Engine/Plugins/Online/OnlineServicesOSSAdapter/Source/Private/Online/AuthOSSAdapter.cpp
--- a/Engine/Plugins/Online/OnlineServicesOSSAdapter/Source/Private/Online/AuthOSSAdapter.cpp
+++ b/Engine/Plugins/Online/OnlineServicesOSSAdapter/Source/Private/Online/AuthOSSAdapter.cpp
@@ -269,7 +269,7 @@ TOnlineAsyncOpHandle<FAuthLogin> FAuthOSSAdapter::Login(FAuthLogin::Params&& Par
 	{
 		const TSharedRef<FAccountInfoOSSAdapter>& AccountInfoOSSAdapter = GetOpDataChecked<TSharedRef<FAccountInfoOSSAdapter>>(InAsyncOp, UE_ONLINE_AUTH_OSS_ADAPTER_KEY_NAME_ACCOUNT_INFO);
 
-		if (!PopulateAttributes(AccountInfoOSSAdapter))
+		if (!PopulateAttributes(*AccountInfoOSSAdapter))
 		{
 			InAsyncOp.SetError(Errors::Unknown());
 
@@ -607,7 +607,7 @@ TOnlineAsyncOpHandle<FAuthHandleLoginStatusChangedImpl> FAuthOSSAdapter::HandleL
 				AccountInfoOSSAdapter->LoginStatus = ELoginStatus::NotLoggedIn;
 				AccountInfoOSSAdapter->LocalUserNum = FPlatformMisc::GetUserIndexForPlatformUser(Params.PlatformUserId);
 				AccountInfoOSSAdapter->UniqueNetId = GetIdentityInterface()->GetUniquePlayerId(AccountInfoOSSAdapter->LocalUserNum);
-				if (!PopulateAttributes(AccountInfoOSSAdapter))
+				if (!PopulateAttributes(*AccountInfoOSSAdapter))
 				{
 					InAsyncOp.SetError(Errors::InvalidUser());
 					return;
@@ -650,24 +650,27 @@ TOnlineAsyncOpHandle<FAuthHandleLoginStatusChangedImpl> FAuthOSSAdapter::HandleL
 	return Op->GetHandle();
 }
 
-bool FAuthOSSAdapter::PopulateAttributes(const TSharedPtr<FAccountInfoOSSAdapter>& AccountInfoOSSAdapter) const
+bool FAuthOSSAdapter::PopulateAttributes(FAccountInfoOSSAdapter& AccountInfoOSSAdapter) const
 {
-	TSharedPtr<FUserOnlineAccount> UserOnlineAccount = GetIdentityInterface()->GetUserAccount(*AccountInfoOSSAdapter->UniqueNetId);
-	if (UserOnlineAccount)
+	if (AccountInfoOSSAdapter.UniqueNetId && AccountInfoOSSAdapter.UniqueNetId->IsValid())
 	{
-		AccountInfoOSSAdapter->Attributes.Emplace(AccountAttributeData::DisplayName, UserOnlineAccount->GetDisplayName());
-		return true;
-	}
+		TSharedPtr<FUserOnlineAccount> UserOnlineAccount = GetIdentityInterface()->GetUserAccount(*AccountInfoOSSAdapter.UniqueNetId);
+		if (UserOnlineAccount)
+		{
+			AccountInfoOSSAdapter.Attributes.Emplace(AccountAttributeData::DisplayName, UserOnlineAccount->GetDisplayName());
+			return true;
+		}
 
-	FString Displayname = GetIdentityInterface()->GetPlayerNickname(*AccountInfoOSSAdapter->UniqueNetId);
+		FString Displayname = GetIdentityInterface()->GetPlayerNickname(*AccountInfoOSSAdapter.UniqueNetId);
 
-	if (!Displayname.IsEmpty())
-	{
-		AccountInfoOSSAdapter->Attributes.Emplace(AccountAttributeData::DisplayName, GetIdentityInterface()->GetPlayerNickname(*AccountInfoOSSAdapter->UniqueNetId));
-		return true;
+		if (!Displayname.IsEmpty())
+		{
+			AccountInfoOSSAdapter.Attributes.Emplace(AccountAttributeData::DisplayName, GetIdentityInterface()->GetPlayerNickname(*AccountInfoOSSAdapter.UniqueNetId));
+			return true;
+		}
 	}
 
-	UE_LOG(LogOnlineServices, Warning, TEXT("[FAuthOSSAdapter::Login][%s] Failure: Failed to find UserOnlineAccount for account %s"), *GetSubsystem().GetSubsystemName().ToString(), *ToLogString(AccountInfoOSSAdapter->AccountId));
+	UE_LOG(LogOnlineServices, Warning, TEXT("[FAuthOSSAdapter::Login][%s] Failure: Failed to find UserOnlineAccount for account %s"), *GetSubsystem().GetSubsystemName().ToString(), *ToLogString(AccountInfoOSSAdapter.AccountId));
 
 	return false;
 }
diff --git a/Engine/Plugins/Online/OnlineServicesOSSAdapter/Source/Public/Online/AuthOSSAdapter.h b/Engine/Plugins/Online/OnlineServicesOSSAdapter/Source/Public/Online/AuthOSSAdapter.h
--- a/Engine/Plugins/Online/OnlineServicesOSSAdapter/Source/Public/Online/AuthOSSAdapter.h
+++ b/Engine/Plugins/Online/OnlineServicesOSSAdapter/Source/Public/Online/AuthOSSAdapter.h
@@ -87,7 +87,7 @@ protected:
 
 	TOnlineAsyncOpHandle<FAuthHandleLoginStatusChangedImpl> HandleLoginStatusChangedImplOp(FAuthHandleLoginStatusChangedImpl::Params&& Params);
 
-	bool PopulateAttributes(const TSharedPtr<FAccountInfoOSSAdapter>& AccountInfoOSSAdapter) const;
+	bool PopulateAttributes(FAccountInfoOSSAdapter& AccountInfoOSSAdapter) const;
 
 	FAccountInfoRegistryOSSAdapter AccountInfoRegistryOSSAdapter;
 	FDelegateHandle OnLoginStatusChangedHandle[MAX_LOCAL_PLAYERS];
