OnlineServices Presence and DisplayName related changes

diff --git a/Engine/Plugins/Online/OnlineServices/Source/OnlineServicesInterface/Private/Online/Presence.cpp b/Engine/Plugins/Online/OnlineServices/Source/OnlineServicesInterface/Private/Online/Presence.cpp
--- a/Engine/Plugins/Online/OnlineServices/Source/OnlineServicesInterface/Private/Online/Presence.cpp
+++ b/Engine/Plugins/Online/OnlineServices/Source/OnlineServicesInterface/Private/Online/Presence.cpp
@@ -51,6 +51,33 @@ void LexFromString(EUserPresenceStatus& OutStatus, const TCHAR* InStr)
 	}
 }
 
+FText ToLocText(EUserPresenceStatus EnumVal)
+{
+	static FText OnlineText =  NSLOCTEXT("OnlinePresence", "Online", "Online");
+	static FText OfflineText =  NSLOCTEXT("OnlinePresence", "Offline", "Offline");
+	static FText AwayText =  NSLOCTEXT("OnlinePresence", "Away", "Away");
+	static FText DoNotDisturbText =  NSLOCTEXT("OnlinePresence", "DoNotDisturb", "Do Not Disturb");
+	static FText UnknownText =  NSLOCTEXT("OnlinePresence", "Unknown", "Unknown");
+
+	switch (EnumVal)
+	{
+	case EUserPresenceStatus::Online:
+		return OnlineText;
+	case EUserPresenceStatus::Offline:
+		return OfflineText;
+	case EUserPresenceStatus::ExtendedAway:
+		// falls through to return away text
+	case EUserPresenceStatus::Away:
+		return AwayText;
+	case EUserPresenceStatus::DoNotDisturb:
+		return DoNotDisturbText;
+	case EUserPresenceStatus::Unknown:
+		return UnknownText;
+	}
+
+	return FText::GetEmpty();
+}
+
 const TCHAR* LexToString(EUserPresenceJoinability EnumVal)
 {
 	switch (EnumVal)
diff --git a/Engine/Plugins/Online/OnlineServices/Source/OnlineServicesInterface/Public/Online/Presence.h b/Engine/Plugins/Online/OnlineServices/Source/OnlineServicesInterface/Public/Online/Presence.h
--- a/Engine/Plugins/Online/OnlineServices/Source/OnlineServicesInterface/Public/Online/Presence.h
+++ b/Engine/Plugins/Online/OnlineServices/Source/OnlineServicesInterface/Public/Online/Presence.h
@@ -49,6 +49,7 @@ enum class EUserPresenceStatus : uint8
 };
 ONLINESERVICESINTERFACE_API const TCHAR* LexToString(EUserPresenceStatus EnumVal);
 ONLINESERVICESINTERFACE_API void LexFromString(EUserPresenceStatus& OutStatus, const TCHAR* InStr);
+ONLINESERVICESINTERFACE_API FText ToLocText(EUserPresenceStatus EnumVal);
 
 enum class EUserPresenceJoinability : uint8
 {
diff --git a/Engine/Plugins/Online/OnlineServices/Source/OnlineServicesInterface/Public/Online/UserInfo.h b/Engine/Plugins/Online/OnlineServices/Source/OnlineServicesInterface/Public/Online/UserInfo.h
--- a/Engine/Plugins/Online/OnlineServices/Source/OnlineServicesInterface/Public/Online/UserInfo.h
+++ b/Engine/Plugins/Online/OnlineServices/Source/OnlineServicesInterface/Public/Online/UserInfo.h
@@ -52,6 +52,8 @@ struct FGetUserInfo
 		FAccountId LocalAccountId;
 		/* User id to get info for*/
 		FAccountId AccountId;
+		/* Preferred platform for DisplayName*/
+		TOptional<FString> DisplayNamePlatform;
 	};
 
 	/**
@@ -199,7 +201,8 @@ END_ONLINE_STRUCT_META()
 
 BEGIN_ONLINE_STRUCT_META(FGetUserInfo::Params)
 	ONLINE_STRUCT_FIELD(FGetUserInfo::Params, LocalAccountId),
-	ONLINE_STRUCT_FIELD(FGetUserInfo::Params, AccountId)
+	ONLINE_STRUCT_FIELD(FGetUserInfo::Params, AccountId),
+	ONLINE_STRUCT_FIELD(FGetUserInfo::Params, DisplayNamePlatform)
 END_ONLINE_STRUCT_META()
 
 BEGIN_ONLINE_STRUCT_META(FGetUserInfo::Result)
diff --git a/Engine/Plugins/Online/OnlineServicesEOS/Source/Private/Online/PresenceEOS.cpp b/Engine/Plugins/Online/OnlineServicesEOS/Source/Private/Online/PresenceEOS.cpp
--- a/Engine/Plugins/Online/OnlineServicesEOS/Source/Private/Online/PresenceEOS.cpp
+++ b/Engine/Plugins/Online/OnlineServicesEOS/Source/Private/Online/PresenceEOS.cpp
@@ -554,10 +554,16 @@ TOnlineAsyncOpHandle<FPartialUpdatePresence> FPresenceEOS::PartialUpdatePresence
 				if (Data->ResultCode == EOS_EResult::EOS_Success)
 				{
 					// Update local presence
+					bool bPresenceHasChanged = false;
 					const FPartialUpdatePresence::Params& Params = InAsyncOp.GetParams();
 					TSharedRef<FUserPresence> LocalUserPresence = FindOrCreatePresence(Params.LocalAccountId, Params.LocalAccountId);
 					if (Params.Mutations.Status.IsSet())
 					{
+						if (LocalUserPresence->Status != Params.Mutations.Status.GetValue())
+						{
+							bPresenceHasChanged = true;
+						}
+
 						LocalUserPresence->Status = Params.Mutations.Status.GetValue();
 					}
 					if (Params.Mutations.StatusString.IsSet())
@@ -581,6 +587,12 @@ TOnlineAsyncOpHandle<FPartialUpdatePresence> FPresenceEOS::PartialUpdatePresence
 						LocalUserPresence->Properties.Emplace(UpdatedProperty.Key, UpdatedProperty.Value);
 					}
 
+					if (bPresenceHasChanged)
+					{
+						FPresenceUpdated PresenceUpdatedParams = { Params.LocalAccountId, LocalUserPresence };
+						OnPresenceUpdatedEvent.Broadcast(PresenceUpdatedParams);
+					}
+
 					InAsyncOp.SetResult(FPartialUpdatePresence::Result());
 				}
 				else
diff --git a/Engine/Plugins/Online/OnlineServicesEOS/Source/Private/Online/UserInfoEOS.cpp b/Engine/Plugins/Online/OnlineServicesEOS/Source/Private/Online/UserInfoEOS.cpp
--- a/Engine/Plugins/Online/OnlineServicesEOS/Source/Private/Online/UserInfoEOS.cpp
+++ b/Engine/Plugins/Online/OnlineServicesEOS/Source/Private/Online/UserInfoEOS.cpp
@@ -156,7 +156,23 @@ TOnlineResult<FGetUserInfo> FUserInfoEOS::GetUserInfo(FGetUserInfo::Params&& Par
 		UE_EOS_CHECK_API_MISMATCH(EOS_USERINFO_COPYBESTDISPLAYNAMEWITHPLATFORM_API_LATEST, 1);
 		WithPlatformOptions.LocalUserId = GetEpicAccountIdChecked(Params.LocalAccountId);
 		WithPlatformOptions.TargetUserId = TargetUserEasId;
-		WithPlatformOptions.TargetPlatformType = EOS_OPT_Epic;
+
+		EOS_OnlinePlatformType OnlinePlatformType = EOS_OPT_Epic;
+		if (Params.DisplayNamePlatform.IsSet())
+		{
+			const FString& DisplayNamePlatform = Params.DisplayNamePlatform.GetValue();
+			EOS_OnlinePlatformType NewOnlinePlatformType = EOSOnlinePlatformTypeFromString(DisplayNamePlatform);
+			if (NewOnlinePlatformType != EOS_OPT_Unknown)
+			{
+				OnlinePlatformType = NewOnlinePlatformType;
+			}
+			else
+			{
+				UE_LOG(LogOnlineServices, VeryVerbose, TEXT("FGetUserInfo::Params::DisplayNamePlatform set to %s. Not a valid Online Platform Type. Epic will be used as default."), *DisplayNamePlatform);
+			}
+		}
+		
+		WithPlatformOptions.TargetPlatformType = OnlinePlatformType;
 
 		EosResult = EOS_UserInfo_CopyBestDisplayNameWithPlatform(UserInfoHandle, &WithPlatformOptions, &EosBestDisplayName);
 	}
diff --git a/Engine/Plugins/Online/OnlineServicesOSSAdapter/Source/Private/Online/PresenceOSSAdapter.cpp b/Engine/Plugins/Online/OnlineServicesOSSAdapter/Source/Private/Online/PresenceOSSAdapter.cpp
--- a/Engine/Plugins/Online/OnlineServicesOSSAdapter/Source/Private/Online/PresenceOSSAdapter.cpp
+++ b/Engine/Plugins/Online/OnlineServicesOSSAdapter/Source/Private/Online/PresenceOSSAdapter.cpp
@@ -10,12 +10,11 @@
 
 #include "OnlineSubsystem.h"
 #include "Interfaces/OnlineIdentityInterface.h"
-#include "Interfaces/OnlinePresenceInterface.h"
 
 
 namespace UE::Online {
 
-EUserPresenceStatus StatusV1toV2(EOnlinePresenceState::Type Status)
+EUserPresenceStatus FPresenceOSSAdapter::StatusV1toV2(EOnlinePresenceState::Type Status)
 {
 	switch (Status)
 	{
@@ -34,7 +33,7 @@ EUserPresenceStatus StatusV1toV2(EOnlinePresenceState::Type Status)
 	}
 }
 
-EOnlinePresenceState::Type StatusV2toV1(EUserPresenceStatus Status)
+EOnlinePresenceState::Type FPresenceOSSAdapter::StatusV2toV1(EUserPresenceStatus Status)
 {
 	switch (Status)
 	{
@@ -87,11 +86,11 @@ void PropertiesV2toV1(TMap<FString, FPresenceVariant> PropertiesIn, FOnlineKeyVa
 	}
 }
 
-TSharedRef<FUserPresence> FPresenceOSSAdapter::PresenceV1toV2(FOnlineUserPresence& Presence)
+TSharedRef<FUserPresence> FPresenceOSSAdapter::PresenceV1toV2(FAccountId AccountId, FOnlineUserPresence& Presence)
 {
 	TSharedRef<FUserPresence> PresenceV2 = MakeShared<FUserPresence>();
 
-	PresenceV2->AccountId = Auth->GetAccountId(Presence.SessionId.ToSharedRef());
+	PresenceV2->AccountId = AccountId;
 	PresenceV2->Status = StatusV1toV2(Presence.Status.State);
 	PresenceV2->GameStatus = GameStatusV1toV2(Presence);
 	PresenceV2->Joinability = (Presence.bIsJoinable) ? EUserPresenceJoinability::Public : EUserPresenceJoinability::Private;
@@ -106,7 +105,6 @@ TSharedRef<FOnlineUserPresence> FPresenceOSSAdapter::PresenceV2toV1(FUserPresenc
 {
 	TSharedRef<FOnlineUserPresence> PresenceV1 = MakeShared<FOnlineUserPresence>();
 
-	PresenceV1->SessionId = Auth->GetUniqueNetId(Presence.AccountId);
 	PresenceV1->bIsOnline = (Presence.Status == EUserPresenceStatus::Online) ? 1 : 0;
 	PresenceV1->bIsPlaying = (Presence.GameStatus != EUserPresenceGameStatus::Unknown) ? 1 : 0;
 	PresenceV1->bIsPlayingThisGame = (Presence.GameStatus == EUserPresenceGameStatus::PlayingThisGame) ? 1 : 0;
@@ -132,30 +130,40 @@ TOnlineAsyncOpHandle<FQueryPresence> FPresenceOSSAdapter::QueryPresence(FQueryPr
 {
 	TOnlineAsyncOpRef<FQueryPresence> Op = GetOp<FQueryPresence>(MoveTemp(Params));
 
-	Op->Then([this](TOnlineAsyncOp<FQueryPresence>& Op, TPromise<TOnlineResult<FQueryPresence>>&& Result)
+	Op->Then([this](TOnlineAsyncOp<FQueryPresence>& Op)
 	{
-		const FUniqueNetIdPtr UniqueNetId = Auth->GetUniqueNetId(Op.GetParams().LocalAccountId);
+		TPromise<void> Promise;
+		TFuture<void> Future = Promise.GetFuture();
 
+		const FUniqueNetIdPtr UniqueNetId = Auth->GetUniqueNetId(Op.GetParams().LocalAccountId);
 		if (UniqueNetId)
 		{
-			PresenceInt->QueryPresence(*UniqueNetId, *MakeDelegateAdapter(this, [this, ResultPromise = MoveTemp(Result)](const FUniqueNetId& UserId, const bool bWasSuccessful) mutable
+			PresenceInt->QueryPresence(*UniqueNetId, *MakeDelegateAdapter(this, [this, WeakOp = Op.AsWeak(), Promise = MoveTemp(Promise)](const FUniqueNetId& UserId, const bool bWasSuccessful) mutable
 			{
-				if (bWasSuccessful)
-				{
-					FQueryPresence::Result Result;
-					ResultPromise.EmplaceValue(MoveTemp(Result));
-				}
-				else
+				if (TOnlineAsyncOpPtr<FQueryPresence> StrongOp = WeakOp.Pin())
 				{
-					ResultPromise.EmplaceValue(Errors::Unknown());
+					if (bWasSuccessful)
+					{
+						FQueryPresence::Result Result;
+						StrongOp->SetResult(MoveTemp(Result));
+					}
+					else
+					{
+						StrongOp->SetError(Errors::Unknown());
+					}
 				}
+
+				Promise.EmplaceValue();
 			}));
 		}
 		else
 		{
-			Result.EmplaceValue(Errors::InvalidUser());
+			Promise.EmplaceValue();
 		}
-	});
+
+		return Future;
+	})
+	.Enqueue(GetSerialQueue());
 
 	return Op->GetHandle();
 }
@@ -165,13 +173,17 @@ TOnlineAsyncOpHandle<FBatchQueryPresence> FPresenceOSSAdapter::BatchQueryPresenc
 {
 	TOnlineAsyncOpRef<FBatchQueryPresence> Op = GetOp<FBatchQueryPresence>(MoveTemp(Params));
 
-	Op->Then([this](TOnlineAsyncOp<FBatchQueryPresence>& Op, TPromise<TOnlineResult<FBatchQueryPresence>>&& Result)
+	Op->Then([this](TOnlineAsyncOp<FBatchQueryPresence>& Op)
 	{
+		TPromise<void> Promise;
+		TFuture<void> Future = Promise.GetFuture();
+
 		const FUniqueNetIdPtr LocalUserId = Auth->GetUniqueNetId(Op.GetParams().LocalAccountId);
 		if (!LocalUserId)
 		{
 			Op.SetError(Errors::InvalidUser());
-			return;
+			Promise.EmplaceValue();
+			return Future;
 		}
 
 		TArray<FUniqueNetIdRef> NetIds;
@@ -181,25 +193,34 @@ TOnlineAsyncOpHandle<FBatchQueryPresence> FPresenceOSSAdapter::BatchQueryPresenc
 			if (!TargetUserNetId)
 			{
 				Op.SetError(Errors::InvalidParams());
-				return;
+				Promise.EmplaceValue();
+				return Future;
 			}
 
 			NetIds.Add(TargetUserNetId.ToSharedRef());
 		}
 
-		PresenceInt->QueryPresence(*LocalUserId, NetIds, *MakeDelegateAdapter(this, [this, ResultPromise=MoveTemp(Result)](const FUniqueNetId& UserId, const bool bWasSuccessful) mutable
+		PresenceInt->QueryPresence(*LocalUserId, NetIds, *MakeDelegateAdapter(this, [this, WeakOp = Op.AsWeak(), Promise = MoveTemp(Promise)](const FUniqueNetId& UserId, const bool bWasSuccessful) mutable
 		{
-			if(bWasSuccessful)
-			{
-				FBatchQueryPresence::Result Result;
-				ResultPromise.EmplaceValue(MoveTemp(Result));
-			}
-			else 
+			if (TOnlineAsyncOpPtr<FBatchQueryPresence> StrongOp = WeakOp.Pin())
 			{
-				ResultPromise.EmplaceValue(Errors::Unknown());
+				if (bWasSuccessful)
+				{
+					FBatchQueryPresence::Result Result;
+					StrongOp->SetResult(MoveTemp(Result));
+				}
+				else
+				{
+					StrongOp->SetError(Errors::Unknown());
+				}
 			}
+
+			Promise.EmplaceValue();
 		}));
-	});
+
+		return Future;
+	})
+	.Enqueue(GetSerialQueue());
 
 	return Op->GetHandle();
 }
@@ -209,9 +230,9 @@ TOnlineResult<FGetCachedPresence> FPresenceOSSAdapter::GetCachedPresence(FGetCac
 	const FUniqueNetId& User = *Auth->GetUniqueNetId(Params.LocalAccountId);
 	TSharedPtr<FOnlineUserPresence> Presence = MakeShared<FOnlineUserPresence>();
 	EOnlineCachedResult::Type Result = PresenceInt->GetCachedPresence(User, Presence);
-	if (Result != EOnlineCachedResult::Success)
+	if (Result == EOnlineCachedResult::Success)
 	{
-		TSharedRef<FUserPresence> V2Presence = PresenceV1toV2(*Presence);
+		TSharedRef<FUserPresence> V2Presence = PresenceV1toV2(Params.TargetAccountId, *Presence);
 		return TOnlineResult<FGetCachedPresence>({V2Presence});
 	}
 	else
@@ -224,22 +245,33 @@ TOnlineAsyncOpHandle<FUpdatePresence> FPresenceOSSAdapter::UpdatePresence(FUpdat
 {
 	TOnlineAsyncOpRef<FUpdatePresence> Op = GetOp<FUpdatePresence>(MoveTemp(Params));
 
-	Op->Then([this](TOnlineAsyncOp<FUpdatePresence>& Op, TPromise<TOnlineResult<FUpdatePresence>>&& Result)
+	Op->Then([this](TOnlineAsyncOp<FUpdatePresence>& Op)
 	{
+		TPromise<void> Promise;
+		TFuture<void> Future = Promise.GetFuture();
+
 		TSharedRef<const FOnlineUserPresence> PresenceV1 = PresenceV2toV1(*Op.GetParams().Presence);
-		PresenceInt->SetPresence(*Auth->GetUniqueNetId(Op.GetParams().LocalAccountId), PresenceV1->Status, *MakeDelegateAdapter(this, [this, ResultPromise = MoveTemp(Result)](const FUniqueNetId& UserId, const bool bWasSuccessful) mutable
+		PresenceInt->SetPresence(*Auth->GetUniqueNetId(Op.GetParams().LocalAccountId), PresenceV1->Status, *MakeDelegateAdapter(this, [this, WeakOp = Op.AsWeak(), Promise = MoveTemp(Promise)](const FUniqueNetId& UserId, const bool bWasSuccessful) mutable
 		{
-			if (bWasSuccessful)
+			if (TOnlineAsyncOpPtr<FUpdatePresence> StrongOp = WeakOp.Pin())
 			{
-				FUpdatePresence::Result Result;
-				ResultPromise.EmplaceValue(MoveTemp(Result));
-			}
-			else
-			{
-				ResultPromise.EmplaceValue(Errors::Unknown());
+				if (bWasSuccessful)
+				{
+					FUpdatePresence::Result Result;
+					StrongOp->SetResult(MoveTemp(Result));
+				}
+				else
+				{
+					StrongOp->SetError(Errors::Unknown());
+				}
 			}
+
+			Promise.EmplaceValue();
 		}));
-	});
+
+		return Future;
+	})
+	.Enqueue(GetSerialQueue());
 
 	return Op->GetHandle();
 }
@@ -248,8 +280,11 @@ TOnlineAsyncOpHandle<FPartialUpdatePresence> FPresenceOSSAdapter::PartialUpdateP
 {
 	TOnlineAsyncOpRef<FPartialUpdatePresence> Op = GetOp<FPartialUpdatePresence>(MoveTemp(Params));
 
-	Op->Then([this](TOnlineAsyncOp<FPartialUpdatePresence>& Op, TPromise<TOnlineResult<FPartialUpdatePresence>>&& Result)
+	Op->Then([this](TOnlineAsyncOp<FPartialUpdatePresence>& Op)
 	{
+		TPromise<void> Promise;
+		TFuture<void> Future = Promise.GetFuture();
+
 		const FPartialUpdatePresence::Params::FMutations& Mutations = Op.GetParams().Mutations;
 		FOnlinePresenceSetPresenceParameters Parameters;
 	
@@ -271,19 +306,27 @@ TOnlineAsyncOpHandle<FPartialUpdatePresence> FPresenceOSSAdapter::PartialUpdateP
 			Parameters.Properties = NewProperties;
 		}
 
-		PresenceInt->SetPresence(*Auth->GetUniqueNetId(Op.GetParams().LocalAccountId), MoveTemp(Parameters), *MakeDelegateAdapter(this, [this, ResultPromise = MoveTemp(Result)](const FUniqueNetId& UserId, const bool bWasSuccessful) mutable
+		PresenceInt->SetPresence(*Auth->GetUniqueNetId(Op.GetParams().LocalAccountId), MoveTemp(Parameters), *MakeDelegateAdapter(this, [this, WeakOp = Op.AsWeak(), Promise = MoveTemp(Promise)](const FUniqueNetId& UserId, const bool bWasSuccessful) mutable
 		{
-			if (bWasSuccessful)
+			if (TOnlineAsyncOpPtr<FPartialUpdatePresence> StrongOp = WeakOp.Pin())
 			{
-				FPartialUpdatePresence::Result Result;
-				ResultPromise.EmplaceValue(MoveTemp(Result));
-			}
-			else
-			{
-				ResultPromise.EmplaceValue(Errors::Unknown());
+				if (bWasSuccessful)
+				{
+					FPartialUpdatePresence::Result Result;
+					StrongOp->SetResult(MoveTemp(Result));
+				}
+				else
+				{
+					StrongOp->SetError(Errors::Unknown());
+				}
 			}
+
+			Promise.EmplaceValue();
 		}));
-	});
+
+		return Future;
+	})
+	.Enqueue(GetSerialQueue());
 
 	return Op->GetHandle();
 }
diff --git a/Engine/Plugins/Online/OnlineServicesOSSAdapter/Source/Public/Online/PresenceOSSAdapter.h b/Engine/Plugins/Online/OnlineServicesOSSAdapter/Source/Public/Online/PresenceOSSAdapter.h
--- a/Engine/Plugins/Online/OnlineServicesOSSAdapter/Source/Public/Online/PresenceOSSAdapter.h
+++ b/Engine/Plugins/Online/OnlineServicesOSSAdapter/Source/Public/Online/PresenceOSSAdapter.h
@@ -5,6 +5,7 @@
 #include "Online/PresenceCommon.h"
 #include "AuthOSSAdapter.h"
 #include "OnlineSubsystemTypes.h"
+#include "Interfaces/OnlinePresenceInterface.h"
 
 class IOnlinePresence;
 using IOnlinePresencePtr = TSharedPtr<IOnlinePresence>;
@@ -32,11 +33,14 @@ public:
 	FUniqueNetIdRef GetUniqueNetId(FAccountId AccountId) const;
 	int32 GetLocalUserNum(FAccountId AccountId) const;
 
+	ONLINESERVICESOSSADAPTER_API static EUserPresenceStatus StatusV1toV2(EOnlinePresenceState::Type Status);
+	ONLINESERVICESOSSADAPTER_API static EOnlinePresenceState::Type StatusV2toV1(EUserPresenceStatus Status);
+
 protected:
 	const FAuthOSSAdapter* Auth;
 	IOnlinePresencePtr PresenceInt;
 
-	TSharedRef<FUserPresence> PresenceV1toV2(FOnlineUserPresence& Presence);
+	TSharedRef<FUserPresence> PresenceV1toV2(FAccountId Accountid, FOnlineUserPresence& Presence);
 	TSharedRef<FOnlineUserPresence> PresenceV2toV1(FUserPresence& Presence);
 };
 
