UE - OSSv2 : Fix for shutdown order in OSSv2 with Adapters

diff --git a/Engine/Plugins/Online/OnlineServicesOSSAdapter/Source/Private/Online/OnlineServicesOSSAdapterModule.cpp b/Engine/Plugins/Online/OnlineServicesOSSAdapter/Source/Private/Online/OnlineServicesOSSAdapterModule.cpp
--- a/Engine/Plugins/Online/OnlineServicesOSSAdapter/Source/Private/Online/OnlineServicesOSSAdapterModule.cpp
+++ b/Engine/Plugins/Online/OnlineServicesOSSAdapter/Source/Private/Online/OnlineServicesOSSAdapterModule.cpp
@@ -20,6 +20,7 @@ public:
 	virtual void StartupModule() override;
 	virtual void ShutdownModule() override;
 protected:
+	TArray<struct FOSSAdapterService> CachedServices;
 };
 
 struct FOSSAdapterService
@@ -27,7 +28,7 @@ struct FOSSAdapterService
 	EOnlineServices Service = EOnlineServices::Default;
 	FString ConfigName;
 	FName OnlineSubsystem;
-	int Priority = -1;
+	int32 Priority = -1;
 };
 
 struct FOSSAdapterConfig
@@ -82,10 +83,22 @@ void FOnlineServicesOSSAdapterModule::StartupModule()
 	FOSSAdapterConfig Config;
 	if (LoadConfig(ConfigProvider, TEXT("OnlineServices.OSSAdapter"), Config))
 	{
-		for (const FOSSAdapterService& ServiceConfig : Config.Services)
+		CachedServices = Config.Services;
+		FModuleManager& ModuleManager = FModuleManager::Get();
+		if (!CachedServices.IsEmpty())
+		{
+			ModuleManager.LoadModule(TEXT("OnlineSubsystem"));
+		}
+
+		for (const FOSSAdapterService& ServiceConfig : CachedServices)
 		{
 			if (IOnlineSubsystem::IsEnabled(ServiceConfig.OnlineSubsystem))
 			{
+				const FName SubsystemName(FString::Printf(TEXT("OnlineSubsystem%s"), *ServiceConfig.OnlineSubsystem.ToString()));
+				if (!ModuleManager.IsModuleLoaded(SubsystemName))
+				{
+					ModuleManager.LoadModule(SubsystemName);
+				}
 				FOnlineServicesRegistry::Get().RegisterServicesFactory(ServiceConfig.Service, MakeUnique<FOnlineServicesFactoryOSSAdapter>(ServiceConfig), ServiceConfig.Priority);
 				FOnlineIdRegistryRegistry::Get().RegisterAccountIdRegistry(ServiceConfig.Service, new FOnlineAccountIdRegistryOSSAdapter(ServiceConfig.Service), ServiceConfig.Priority);
 				FOnlineIdRegistryRegistry::Get().RegisterSessionIdRegistry(ServiceConfig.Service, new FOnlineSessionIdRegistryOSSAdapter(ServiceConfig.Service), ServiceConfig.Priority);
@@ -97,18 +110,14 @@ void FOnlineServicesOSSAdapterModule::StartupModule()
 
 void FOnlineServicesOSSAdapterModule::ShutdownModule()
 {
-	FOnlineConfigProviderGConfig ConfigProvider(GEngineIni);
-	FOSSAdapterConfig Config;
-	if (LoadConfig(ConfigProvider, TEXT("OnlineServices.OSSAdapter"), Config))
+	for (FOSSAdapterService& ServiceConfig : CachedServices)
 	{
-		for (const FOSSAdapterService& ServiceConfig : Config.Services)
-		{
-			FOnlineServicesRegistry::Get().UnregisterServicesFactory(ServiceConfig.Service, ServiceConfig.Priority);
-			FOnlineIdRegistryRegistry::Get().UnregisterAccountIdRegistry(ServiceConfig.Service, ServiceConfig.Priority);
-			FOnlineIdRegistryRegistry::Get().UnregisterSessionIdRegistry(ServiceConfig.Service, ServiceConfig.Priority);
-			FOnlineIdRegistryRegistry::Get().UnregisterSessionInviteIdRegistry(ServiceConfig.Service, ServiceConfig.Priority);
-		}
+		FOnlineServicesRegistry::Get().UnregisterServicesFactory(ServiceConfig.Service, ServiceConfig.Priority);
+		FOnlineIdRegistryRegistry::Get().UnregisterAccountIdRegistry(ServiceConfig.Service, ServiceConfig.Priority);
+		FOnlineIdRegistryRegistry::Get().UnregisterSessionIdRegistry(ServiceConfig.Service, ServiceConfig.Priority);
+		FOnlineIdRegistryRegistry::Get().UnregisterSessionInviteIdRegistry(ServiceConfig.Service, ServiceConfig.Priority);
 	}
+	CachedServices.Empty();
 }
 
 /* UE::Online */ }
