Fix client crash on shutdown when the default online subsystem is reloaded, invalidating the OSS adapter reference.
#rnx
#Synced-CL 50276648

diff --git a/Engine/Plugins/Online/OnlineServicesOSSAdapter/Source/Private/Online/OnlineServicesOSSAdapter.cpp b/Engine/Plugins/Online/OnlineServicesOSSAdapter/Source/Private/Online/OnlineServicesOSSAdapter.cpp
--- a/Engine/Plugins/Online/OnlineServicesOSSAdapter/Source/Private/Online/OnlineServicesOSSAdapter.cpp
+++ b/Engine/Plugins/Online/OnlineServicesOSSAdapter/Source/Private/Online/OnlineServicesOSSAdapter.cpp
@@ -19,6 +19,8 @@
 #include "Online/CommerceOSSAdapter.h"
 
 #include "OnlineSubsystem.h"
+#include "Online/OnlineServicesOSSAdapterModule.h"
+#include "Online/OnlineServicesRegistry.h"
 
 namespace UE::Online {
 
@@ -87,7 +89,12 @@ void FOnlineServicesOSSAdapter::RegisterComponents()
 void FOnlineServicesOSSAdapter::Initialize()
 {
 	AccountIdRegistry = static_cast<FOnlineAccountIdRegistryOSSAdapter*>(FOnlineIdRegistryRegistry::Get().GetAccountIdRegistry(ServicesType));
-
+	
+	if (FOnlineServicesOSSAdapterModule* OSSAdapterModule = FModuleManager::GetModulePtr<FOnlineServicesOSSAdapterModule>("OnlineServicesOSSAdapter"))
+	{
+		OSSAdapterModule->RegisterOSSAdapter(StaticCastSharedPtr<FOnlineServicesOSSAdapter>(AsShared().ToSharedPtr()));
+	} 
+	
 	FOnlineServicesCommon::Initialize();
 }
 
diff --git a/Engine/Plugins/Online/OnlineServicesOSSAdapter/Source/Private/Online/OnlineServicesOSSAdapterModule.cpp b/Engine/Plugins/Online/OnlineServicesOSSAdapter/Source/Private/Online/OnlineServicesOSSAdapterModule.cpp
--- a/Engine/Plugins/Online/OnlineServicesOSSAdapter/Source/Private/Online/OnlineServicesOSSAdapterModule.cpp
+++ b/Engine/Plugins/Online/OnlineServicesOSSAdapter/Source/Private/Online/OnlineServicesOSSAdapterModule.cpp
@@ -1,6 +1,9 @@
 // Copyright Epic Games, Inc. All Rights Reserved.
 
+#include "Online/OnlineServicesOSSAdapterModule.h"
+
 #include "CoreMinimal.h"
+#include "OnlineDelegates.h"
 #include "Modules/ModuleInterface.h"
 #include "Modules/ModuleManager.h"
 
@@ -14,15 +17,6 @@
 namespace UE::Online
 {
 
-class FOnlineServicesOSSAdapterModule : public IModuleInterface
-{
-public:
-	virtual void StartupModule() override;
-	virtual void ShutdownModule() override;
-protected:
-	TArray<struct FOSSAdapterService> CachedServices;
-};
-
 struct FOSSAdapterService
 {
 	EOnlineServices Service = EOnlineServices::Default;
@@ -106,6 +100,8 @@ void FOnlineServicesOSSAdapterModule::StartupModule()
 			}
 		}
 	}
+	
+	OnSubsystemPreReloadHandle = FOnlineSubsystemDelegates::OnDefaultOnlineSubsystemPreReloaded.AddRaw(this, &FOnlineServicesOSSAdapterModule::OnSubsystemPreReload);
 }
 
 void FOnlineServicesOSSAdapterModule::ShutdownModule()
@@ -118,6 +114,43 @@ void FOnlineServicesOSSAdapterModule::ShutdownModule()
 		FOnlineIdRegistryRegistry::Get().UnregisterSessionInviteIdRegistry(ServiceConfig.Service, ServiceConfig.Priority);
 	}
 	CachedServices.Empty();
+	FOnlineSubsystemDelegates::OnDefaultOnlineSubsystemPreReloaded.Remove(OnSubsystemPreReloadHandle);
+	OnSubsystemPreReloadHandle.Reset();
+}
+	
+void FOnlineServicesOSSAdapterModule::GetAllAdapters(TArray<TSharedPtr<FOnlineServicesOSSAdapter>>& OutOSSAdapters)
+{
+	for (int32 AdapterIdx = OSSAdapters.Num() - 1; AdapterIdx >= 0; AdapterIdx--)
+	{
+		if (TSharedPtr<FOnlineServicesOSSAdapter> OSSAdapter = OSSAdapters[AdapterIdx].Pin())
+		{
+			OutOSSAdapters.Add(OSSAdapter);
+		}
+		else
+		{
+			OSSAdapters.RemoveAtSwap(AdapterIdx);
+		}
+	}
+}
+	
+void FOnlineServicesOSSAdapterModule::RegisterOSSAdapter(const TSharedPtr<FOnlineServicesOSSAdapter> OnlineServicesOSSAdapter)
+{
+	OSSAdapters.Add(OnlineServicesOSSAdapter);
+}
+	
+void FOnlineServicesOSSAdapterModule::OnSubsystemPreReload(IOnlineSubsystem* OnlineSubsystem)
+{
+	TArray<TSharedPtr<FOnlineServicesOSSAdapter>> OutOnlineServices;
+	GetAllAdapters(OutOnlineServices);
+	
+	for (const TSharedPtr<FOnlineServicesOSSAdapter>& OnlineServiceOSSAdapter : OutOnlineServices)
+	{
+		if (&OnlineServiceOSSAdapter->GetSubsystem() == OnlineSubsystem)
+		{
+			DestroyService(OnlineServiceOSSAdapter->GetServicesProvider(), OnlineServiceOSSAdapter->GetInstanceName(),
+				OnlineServiceOSSAdapter->GetInstanceConfigName());
+		}
+	}
 }
 
 /* UE::Online */ }
diff --git a/Engine/Plugins/Online/OnlineServicesOSSAdapter/Source/Public/Online/OnlineServicesOSSAdapterModule.h b/Engine/Plugins/Online/OnlineServicesOSSAdapter/Source/Public/Online/OnlineServicesOSSAdapterModule.h
new file mode 100644
--- /dev/null
+++ b/Engine/Plugins/Online/OnlineServicesOSSAdapter/Source/Public/Online/OnlineServicesOSSAdapterModule.h
@@ -0,0 +1,40 @@
+﻿// Copyright Epic Games, Inc. All Rights Reserved.
+
+#pragma once
+
+#include "OnlineSubsystem.h"
+#include "Modules/ModuleInterface.h"
+
+namespace UE::Online {
+	
+	class FOnlineServicesOSSAdapter;
+
+	class FOnlineServicesOSSAdapterModule : public IModuleInterface
+	{
+	public:
+		virtual void StartupModule() override;
+		virtual void ShutdownModule() override;
+		
+		/**
+		 * Get list of all instantiated OnlineServicesAdapter
+		 * 
+		 * @param OutOSSAdapters Array of online services adapters to fill
+		 */
+		ONLINESERVICESOSSADAPTER_API void GetAllAdapters(TArray<TSharedPtr<FOnlineServicesOSSAdapter>>& OutOSSAdapters);
+	
+		/**
+		 * Register a new OnlineServicesAdapter to keep track of it
+		 * 
+		 * @param OnlineServicesOSSAdapter Online service adapter to track
+		 */
+		ONLINESERVICESOSSADAPTER_API void RegisterOSSAdapter(const TSharedPtr<FOnlineServicesOSSAdapter> OnlineServicesOSSAdapter);
+		
+	protected:
+		void OnSubsystemPreReload(IOnlineSubsystem* OnlineSubsystem);
+	
+		FDelegateHandle OnSubsystemPreReloadHandle;
+		TArray<TWeakPtr<FOnlineServicesOSSAdapter>> OSSAdapters;
+		TArray<struct FOSSAdapterService> CachedServices;
+	};
+	
+	/*UE::Online*/ }
diff --git a/Engine/Plugins/Online/OnlineSubsystem/Source/Private/OnlineSubsystemModule.cpp b/Engine/Plugins/Online/OnlineSubsystem/Source/Private/OnlineSubsystemModule.cpp
--- a/Engine/Plugins/Online/OnlineSubsystem/Source/Private/OnlineSubsystemModule.cpp
+++ b/Engine/Plugins/Online/OnlineSubsystem/Source/Private/OnlineSubsystemModule.cpp
@@ -17,6 +17,7 @@ namespace
  * OnlineDelegates.h static declarations
  */
 FOnlineSubsystemDelegates::FOnOnlineSubsystemCreated FOnlineSubsystemDelegates::OnOnlineSubsystemCreated;
+FOnlineSubsystemDelegates::FOnDefaultOnlineSubsystemPreReloaded FOnlineSubsystemDelegates::OnDefaultOnlineSubsystemPreReloaded;
 
 static inline FName GetOnlineModuleName(const FString& SubsystemName, const TMap<FString, FName>& ModuleRedirects)
 {
@@ -227,6 +228,10 @@ void FOnlineSubsystemModule::LoadDefaultSubsystem()
 
 void FOnlineSubsystemModule::ReloadDefaultSubsystem()
 {
+	if (IOnlineSubsystem* OSS = GetOnlineSubsystem(DefaultPlatformService))
+	{
+		FOnlineSubsystemDelegates::OnDefaultOnlineSubsystemPreReloaded.Broadcast(OSS);
+	}
 	DestroyOnlineSubsystem(DefaultPlatformService);
 	// Clear our InstanceNames cache so we can re-establish it in case the DefaultPlatformService
 	InstanceNames.Empty();
diff --git a/Engine/Plugins/Online/OnlineSubsystem/Source/Public/OnlineDelegates.h b/Engine/Plugins/Online/OnlineSubsystem/Source/Public/OnlineDelegates.h
--- a/Engine/Plugins/Online/OnlineSubsystem/Source/Public/OnlineDelegates.h
+++ b/Engine/Plugins/Online/OnlineSubsystem/Source/Public/OnlineDelegates.h
@@ -22,6 +22,14 @@ public:
 	 */
 	DECLARE_MULTICAST_DELEGATE_OneParam(FOnOnlineSubsystemCreated, class IOnlineSubsystem* /*NewSubsystem*/);
 	static UE_API FOnOnlineSubsystemCreated OnOnlineSubsystemCreated;
+	
+	/**
+	 * Notification that the default subsystem instance is going to be deleted and reinstantiated
+	 * 
+	 * @param SubsystemToBeReloaded the instance to be destroyed
+	 */
+	DECLARE_MULTICAST_DELEGATE_OneParam(FOnDefaultOnlineSubsystemPreReloaded, class IOnlineSubsystem* /*SubsystemToBeReloaded*/);
+	static UE_API FOnDefaultOnlineSubsystemPreReloaded OnDefaultOnlineSubsystemPreReloaded;
 };
 
 #undef UE_API
